Move ConsoleVariablesEditor from Plugins/Experimental to Plugins/Editor

#[fyi] Jason.Walter

diff --git a/Engine/Plugins/Editor/ConsoleVariablesEditor/ConsoleVariables.uplugin b/Engine/Plugins/Editor/ConsoleVariablesEditor/ConsoleVariables.uplugin
new file mode 100644
--- /dev/null
+++ b/Engine/Plugins/Editor/ConsoleVariablesEditor/ConsoleVariables.uplugin
@@ -0,0 +1,40 @@
+{
+	"FileVersion": 1,
+	"Version": 1,
+	"VersionName": "1.0",
+	"FriendlyName": "Console Variables Editor",
+	"Description": "Save, load and control Console Variables (cvars) from this panel using Slate.",
+	"Category": "Virtual Production",
+	"CreatedBy": "",
+	"CreatedByURL": "",
+	"DocsURL": "",
+	"EnabledByDefault": false,
+	"MarketplaceURL": "",
+	"SupportURL": "",
+	"CanContainContent": true,
+	"IsBetaVersion": true,
+	"IsExperimentalVersion": false,
+	"Installed": false,
+	"Plugins":
+	[
+		{
+			"Name": "ConcertSyncClient",
+			"Enabled": true
+		},
+		{
+			"Name": "ConcertSyncCore",
+			"Enabled": true
+		},
+		{
+			"Name": "ConcertMain",
+			"Enabled": true
+		},
+	],
+	"Modules": [
+		{
+			"Name": "ConsoleVariablesEditor",
+			"Type": "UncookedOnly",
+			"LoadingPhase": "Default"
+		}
+	]
+}
diff --git a/Engine/Plugins/Editor/ConsoleVariablesEditor/Source/ConsoleVariablesEditor/ConsoleVariablesEditor.build.cs b/Engine/Plugins/Editor/ConsoleVariablesEditor/Source/ConsoleVariablesEditor/ConsoleVariablesEditor.build.cs
new file mode 100644
--- /dev/null
+++ b/Engine/Plugins/Editor/ConsoleVariablesEditor/Source/ConsoleVariablesEditor/ConsoleVariablesEditor.build.cs
@@ -0,0 +1,48 @@
+// Copyright Epic Games, Inc. All Rights Reserved.
+
+using System.IO;
+using UnrealBuildTool;
+
+public class ConsoleVariablesEditor : ModuleRules
+{
+	public ConsoleVariablesEditor(ReadOnlyTargetRules Target) : base(Target)
+	{
+		PCHUsage = ModuleRules.PCHUsageMode.UseExplicitOrSharedPCHs;
+
+		PublicDependencyModuleNames.AddRange(
+			new string[]
+			{
+				"Core",
+				"OutputLog"
+			}
+			);
+			
+		
+		PrivateDependencyModuleNames.AddRange(
+			new string[]
+			{
+				"Concert",
+				"ConcertSyncClient",
+				"ConcertSyncCore",
+				"ConcertTransport",
+				"AssetRegistry",
+				"AssetTools",
+				"CoreUObject",
+				"ContentBrowser",
+				"Engine",
+				"EditorStyle",
+				"EditorWidgets",
+				"InputCore",
+				"Kismet",
+				"Projects",
+				"PropertyEditor",
+				"Slate",
+				"SlateCore",
+				"ToolMenus", 
+				"ToolWidgets",
+				"UnrealEd",
+				"WorkspaceMenuStructure"
+			}
+			);
+	}
+}
diff --git a/Engine/Plugins/Editor/ConsoleVariablesEditor/Source/ConsoleVariablesEditor/Private/AssetTypeActions/AssetTypeActions_ConsoleVariables.cpp b/Engine/Plugins/Editor/ConsoleVariablesEditor/Source/ConsoleVariablesEditor/Private/AssetTypeActions/AssetTypeActions_ConsoleVariables.cpp
new file mode 100644
--- /dev/null
+++ b/Engine/Plugins/Editor/ConsoleVariablesEditor/Source/ConsoleVariablesEditor/Private/AssetTypeActions/AssetTypeActions_ConsoleVariables.cpp
@@ -0,0 +1,55 @@
+// Copyright Epic Games, Inc. All Rights Reserved.
+
+#include "AssetTypeActions_ConsoleVariables.h"
+
+#include "ConsoleVariablesAsset.h"
+#include "ConsoleVariablesEditor/Private/ConsoleVariablesEditorModule.h"
+
+#include "Framework/MultiBox/MultiBoxBuilder.h"
+
+#define LOCTEXT_NAMESPACE "ConsoleVariablesEditor"
+
+FText FAssetTypeActions_ConsoleVariables::GetName() const
+{
+	return LOCTEXT("AssetTypeActions_ConsoleVariable_Name", "Console Variable Collection");
+}
+
+UClass* FAssetTypeActions_ConsoleVariables::GetSupportedClass() const
+{
+	return UConsoleVariablesAsset::StaticClass();
+}
+
+void FAssetTypeActions_ConsoleVariables::GetActions(const TArray<UObject*>& InObjects, FMenuBuilder& MenuBuilder)
+{
+	FAssetTypeActions_Base::GetActions(InObjects, MenuBuilder);
+
+	const TArray<TWeakObjectPtr<UConsoleVariablesAsset>> ConsoleVariableAssets = GetTypedWeakObjectPtrs<UConsoleVariablesAsset>(InObjects);
+
+	MenuBuilder.AddMenuEntry(
+		LOCTEXT("AssetTypeActions_OpenVariableCollection", "Open Variable Collection in Editor"),
+		LOCTEXT("AssetTypeActions_OpenVariableCollectionToolTip", "Open this console variable collection in the Console Variables Editor. Select only one asset at a time."),
+		FSlateIcon(FAppStyle::Get().GetStyleSetName(), "SystemWideCommands.SummonOpenAssetDialog"),
+		FUIAction(
+			FExecuteAction::CreateLambda([=] {
+					OpenAssetEditor(InObjects, TSharedPtr<IToolkitHost>());
+			}),
+			FCanExecuteAction::CreateLambda([=] {
+					// We only want to open a single Variable Collection asset at a time, so let's ensure
+					// the number of selected assets is exactly one.
+					return ConsoleVariableAssets.Num() == 1;
+				})
+			)
+	);
+}
+
+void FAssetTypeActions_ConsoleVariables::OpenAssetEditor(const TArray<UObject*>& InObjects, TSharedPtr<IToolkitHost> EditWithinLevelEditor)
+{
+	if (InObjects.Num() && InObjects[0])
+	{
+		FConsoleVariablesEditorModule& Module = FConsoleVariablesEditorModule::Get();
+					
+		Module.OpenConsoleVariablesDialogWithAssetSelected(FAssetData(InObjects[0]));
+	}
+}
+
+#undef LOCTEXT_NAMESPACE
diff --git a/Engine/Plugins/Editor/ConsoleVariablesEditor/Source/ConsoleVariablesEditor/Private/AssetTypeActions/AssetTypeActions_ConsoleVariables.h b/Engine/Plugins/Editor/ConsoleVariablesEditor/Source/ConsoleVariablesEditor/Private/AssetTypeActions/AssetTypeActions_ConsoleVariables.h
new file mode 100644
--- /dev/null
+++ b/Engine/Plugins/Editor/ConsoleVariablesEditor/Source/ConsoleVariablesEditor/Private/AssetTypeActions/AssetTypeActions_ConsoleVariables.h
@@ -0,0 +1,76 @@
+// Copyright Epic Games, Inc. All Rights Reserved.
+
+#pragma once
+
+#include "AssetTypeActions_Base.h"
+
+class FAssetTypeActions_ConsoleVariables : public FAssetTypeActions_Base
+{
+
+public:
+	// ~ Begin IAssetTypeActions Interface
+	
+	/**
+	 * @return Whether or not this asset has a live filter
+	 */
+	virtual bool CanFilter() override
+	{
+		return true;
+	}
+	
+	/**
+	 * @return Which category the asset will appear in when creating a new asset in the content browser
+	 */
+	virtual uint32 GetCategories() override
+	{
+		return EAssetTypeCategories::Misc;
+	}
+	
+	/**
+	 * @return The color on the asset image container
+	 */
+	virtual FColor GetTypeColor() const override
+	{
+		return FColor(238, 181, 235, 255);
+	}
+	
+	/**
+	 * @param InObjects A list of objects selected in the content browser
+	 * @return Whether or not the asset has context actions (required to return true in order to use registered actions)
+	 */
+	virtual bool HasActions(const TArray<UObject*>& InObjects) const override
+	{
+		return true;
+	}
+
+	// These are defined in source - GetName requires a defined LOCTEXT namespace and GetSupportedClass has an external include
+	
+	/**
+	 * @return The display name for the asset type
+	 */
+	virtual FText GetName() const override;
+
+	/**
+	 * @return The asset class these actions will belong to
+	 */
+	virtual UClass* GetSupportedClass() const override;
+	
+	/**
+	 * @brief Called when an item in the content browser is right-clicked. This determines what actions will be shown if HasActions() returns true.
+	 * @param InObjects A list of objects selected in the content browser
+	 * @param MenuBuilder A reference to an instance of the class which adds menu items to context menus
+	 */
+	virtual void GetActions(const TArray<UObject*>& InObjects, FMenuBuilder& MenuBuilder) override;
+
+	/**
+	 * @return Whether the toolkit use dto open the asset should be world-centric (and thus should not use a standalone editor window)
+	 */
+	virtual bool ShouldForceWorldCentric() override
+	{
+		return true;
+	}
+
+	virtual void OpenAssetEditor(const TArray<UObject*>& InObjects, TSharedPtr<IToolkitHost> EditWithinLevelEditor) override;
+
+	// ~ End IAssetTypeActions Interface
+};
diff --git a/Engine/Plugins/Editor/ConsoleVariablesEditor/Source/ConsoleVariablesEditor/Private/ConsoleVariablesAsset.cpp b/Engine/Plugins/Editor/ConsoleVariablesEditor/Source/ConsoleVariablesEditor/Private/ConsoleVariablesAsset.cpp
new file mode 100644
--- /dev/null
+++ b/Engine/Plugins/Editor/ConsoleVariablesEditor/Source/ConsoleVariablesEditor/Private/ConsoleVariablesAsset.cpp
@@ -0,0 +1,70 @@
+// Copyright Epic Games, Inc. All Rights Reserved.
+
+#include "ConsoleVariablesAsset.h"
+
+#include "ConsoleVariablesEditorModule.h"
+#include "Algo/Find.h"
+
+void UConsoleVariablesAsset::SetVariableCollectionDescription(const FString& InVariableCollectionDescription)
+{
+	VariableCollectionDescription = InVariableCollectionDescription;
+}
+
+void UConsoleVariablesAsset::ReplaceSavedCommands(const TArray<FConsoleVariablesEditorAssetSaveData>& Replacement)
+{
+	SavedCommands = Replacement;
+}
+
+bool UConsoleVariablesAsset::FindSavedDataByCommandString(const FString InCommandString, FConsoleVariablesEditorAssetSaveData& OutValue) const
+{
+	for (const FConsoleVariablesEditorAssetSaveData& Command : SavedCommands)
+	{
+		if (Command.CommandName.Equals(InCommandString))
+		{
+			OutValue = Command;
+			return true;
+		}
+	}
+	
+	return false;
+}
+
+void UConsoleVariablesAsset::AddOrSetConsoleObjectSavedData(const FConsoleVariablesEditorAssetSaveData& InData)
+{
+	FConsoleVariablesEditorModule& ConsoleVariablesEditorModule = FConsoleVariablesEditorModule::Get();
+
+	if (const TWeakPtr<FConsoleVariablesEditorCommandInfo> CommandInfo =
+			ConsoleVariablesEditorModule.FindCommandInfoByName(InData.CommandName);
+		CommandInfo.IsValid())
+	{
+		const IConsoleVariable* AsVariable = CommandInfo.Pin()->GetConsoleVariablePtr();
+		
+		if (AsVariable && AsVariable->TestFlags(ECVF_RenderThreadSafe))
+		{
+			UE_LOG(LogConsoleVariablesEditor, Warning,
+				TEXT("The console variable named %s is flagged as ECVF_RenderThreadSafe. The value on the render thread will lag behind the value on the main thread by one frame if r.OneFrameThreadLag is 1."),
+				*InData.CommandName
+			);
+		}
+	}
+	
+	RemoveConsoleVariable(InData.CommandName);
+	SavedCommands.Add(InData);
+}
+
+bool UConsoleVariablesAsset::RemoveConsoleVariable(const FString InCommandString)
+{
+	FConsoleVariablesEditorAssetSaveData ExistingData;
+	if (FindSavedDataByCommandString(InCommandString, ExistingData))
+	{
+		return SavedCommands.Remove(ExistingData) > 0;
+	}
+	
+	return false;
+}
+
+void UConsoleVariablesAsset::CopyFrom(const UConsoleVariablesAsset* InAssetToCopy)
+{
+	VariableCollectionDescription = InAssetToCopy->GetVariableCollectionDescription();
+	SavedCommands = InAssetToCopy->GetSavedCommands();
+}
diff --git a/Engine/Plugins/Editor/ConsoleVariablesEditor/Source/ConsoleVariablesEditor/Private/ConsoleVariablesEditorCommandInfo.h b/Engine/Plugins/Editor/ConsoleVariablesEditor/Source/ConsoleVariablesEditor/Private/ConsoleVariablesEditorCommandInfo.h
new file mode 100644
--- /dev/null
+++ b/Engine/Plugins/Editor/ConsoleVariablesEditor/Source/ConsoleVariablesEditor/Private/ConsoleVariablesEditorCommandInfo.h
@@ -0,0 +1,266 @@
+// Copyright Epic Games, Inc. All Rights Reserved.
+
+#pragma once
+
+#include "Algo/Find.h"
+#include "CoreMinimal.h"
+#include "ConsoleVariablesEditorLog.h"
+#include "Editor.h"
+#include "Engine/GameEngine.h"
+#include "HAL/IConsoleManager.h"
+
+DECLARE_MULTICAST_DELEGATE_OneParam(FOnDetectConsoleObjectUnregistered, FString)
+
+#define LOCTEXT_NAMESPACE "ConsoleVariablesEditor"
+
+struct FConsoleVariablesEditorCommandInfo
+{
+	enum class EConsoleObjectType
+	{
+		// A console command that has no associated console object but is parsed externally, e.g. 'stat unit'
+		NullObject,
+		// A console command with an associated console object, like 'r.SetNearClipPlane'
+		Command,
+		// A console variable such as 'r.ScreenPercentage'
+		Variable
+	};
+
+	struct FStaticConsoleVariableFlagInfo
+	{
+		EConsoleVariableFlags Flag;
+		FText DisplayText;
+	};
+
+	FConsoleVariablesEditorCommandInfo(const FString& InCommand)
+	: Command(InCommand)
+	{
+		if (GetConsoleObjectPtr())
+		{
+			ObjectType = EConsoleObjectType::Command;
+			
+			if (const IConsoleVariable* AsVariable = GetConsoleVariablePtr())
+			{
+				ObjectType = EConsoleObjectType::Variable;
+				StartupValueAsString = AsVariable->GetString();
+				StartupSource = GetSource();
+			}
+		}
+	}
+
+	~FConsoleVariablesEditorCommandInfo()
+	{
+		OnDetectConsoleObjectUnregistered.Remove(OnDetectConsoleObjectUnregisteredHandle);
+
+		if (IConsoleVariable* AsVariable = GetConsoleVariablePtr())
+		{
+			AsVariable->OnChangedDelegate().Remove(OnVariableChangedCallbackHandle);
+		}
+	}
+
+	FORCEINLINE bool operator==(const FConsoleVariablesEditorCommandInfo& Comparator) const
+	{
+		return Command.Equals(Comparator.Command);
+	}
+
+	void ExecuteCommand(const FString& NewValueAsString) const
+	{
+		GEngine->Exec(GetCurrentWorld(),
+			*FString::Printf(TEXT("%s %s"), *Command, *NewValueAsString).TrimStartAndEnd());
+	}
+
+	/** Get a reference to the cached console object. May return nullptr if unregistered. */
+	IConsoleObject* GetConsoleObjectPtr()
+	{
+		// If the console object ptr goes stale or is older than the specified threshold, try to refresh it
+		// May return nullptr if unregistered
+		if (!ConsoleObjectPtr ||
+			(FDateTime::UtcNow() - TimeOfLastConsoleObjectRefresh).GetTotalSeconds() > ConsoleObjectRefreshThreshold)
+		{
+			ConsoleObjectPtr = IConsoleManager::Get().FindConsoleObject(*Command);
+			TimeOfLastConsoleObjectRefresh = FDateTime::UtcNow();
+		}
+
+		// If the console object turns out to be unregistered, let interested parties know
+		if (ConsoleObjectPtr && ConsoleObjectPtr->TestFlags(ECVF_Unregistered))
+		{
+			OnDetectConsoleObjectUnregistered.Broadcast(Command);
+		}
+
+		return ConsoleObjectPtr;
+	}
+
+	/** Return the console object as a console variable object if applicable. May return nullptr if unregistered. */
+	IConsoleVariable* GetConsoleVariablePtr()
+	{
+		if (IConsoleObject* ObjectPtr = GetConsoleObjectPtr())
+		{
+			return ObjectPtr->AsVariable();
+		}
+		
+		return nullptr;
+	}
+
+	/**
+	 *Return the console object as a console command object if applicable.
+	 *Does not consider externally parsed console commands, as they have no associated objects.
+	 */
+	IConsoleCommand* GetConsoleCommandPtr()
+	{
+		if (IConsoleObject* ObjectPtr = GetConsoleObjectPtr())
+		{
+			return ObjectPtr->AsCommand();
+		}
+		
+		return nullptr;
+	}
+	
+	static UWorld* GetCurrentWorld()
+	{
+		UWorld* CurrentWorld = nullptr;
+		if (GIsEditor)
+		{
+			CurrentWorld = GEditor->GetEditorWorldContext().World();
+		}
+		else if (UGameEngine* GameEngine = Cast<UGameEngine>(GEngine))
+		{
+			CurrentWorld = GameEngine->GetGameWorld();
+		}
+		return CurrentWorld;
+	}
+
+	FString GetHelpText()
+	{
+		if (const IConsoleVariable* AsVariable = GetConsoleVariablePtr())
+		{
+			return FString(AsVariable->GetHelp());
+		}
+
+		return "";
+	}
+
+	EConsoleVariableFlags GetSource()
+	{
+		if (const IConsoleObject* ConsoleObject = GetConsoleObjectPtr())
+		{
+			return (EConsoleVariableFlags)((uint32)ConsoleObject->GetFlags() & ECVF_SetByMask);
+		}
+
+		return ECVF_Default;
+	}
+
+	void ClearSourceFlags()
+	{
+		if (IConsoleObject* ConsoleObject = GetConsoleObjectPtr())
+		{
+			for (const FStaticConsoleVariableFlagInfo& StaticConsoleVariableFlagInfo : SupportedFlags)
+			{
+				ConsoleObject->ClearFlags(StaticConsoleVariableFlagInfo.Flag);
+			}
+		}
+	}
+
+	void SetSourceFlag(const EConsoleVariableFlags InSource)
+	{
+		const uint32 OldPri = (uint32)GetSource();
+		const uint32 NewPri = (uint32)InSource;
+
+		if (NewPri < OldPri)
+		{
+			return;
+		}
+		
+		if (IConsoleObject* ConsoleObject = GetConsoleObjectPtr())
+		{
+			ClearSourceFlags();
+			ConsoleObject->SetFlags((EConsoleVariableFlags)InSource);
+		}
+	}
+
+	FText GetSourceAsText()
+	{
+		return ConvertConsoleVariableSetByFlagToText(GetSource());
+	}
+
+	static FText ConvertConsoleVariableSetByFlagToText(const EConsoleVariableFlags InFlag)
+	{
+		FText ReturnValue = LOCTEXT("UnknownSource", "<UNKNOWN>"); 
+
+		if (const FStaticConsoleVariableFlagInfo* Match = Algo::FindByPredicate(
+			SupportedFlags,
+				[InFlag](const FStaticConsoleVariableFlagInfo& Comparator)
+				{
+					return Comparator.Flag == InFlag;
+				}))
+		{
+			ReturnValue = (*Match).DisplayText;
+		}
+		
+		return ReturnValue;
+	}
+
+	bool IsCurrentValueDifferentFromInputValue(const FString& InValueToCompare)
+	{
+		if (const IConsoleVariable* AsVariable = GetConsoleVariablePtr())
+		{
+			// Floats sometimes return true erroneously because they can be stringified as e.g '1' or '1.0' by different functions.
+			if (AsVariable->IsVariableFloat())
+			{
+				const float A = AsVariable->GetFloat();
+				const float B = FCString::Atof(*InValueToCompare);
+
+				return !FMath::IsNearlyEqual(A, B);
+			}
+			
+			return !AsVariable->GetString().Equals(InValueToCompare);
+		}
+		else if (ObjectType == EConsoleObjectType::NullObject || ObjectType == EConsoleObjectType::Command)
+		{
+			return true;
+		}
+
+		return false;
+	}
+
+	/** The actual string key or name */
+	UPROPERTY()
+	FString Command;
+
+	EConsoleObjectType ObjectType = EConsoleObjectType::NullObject;
+
+	/** This object is periodically refreshed to mitigate the occurrence of stale pointers. */
+	IConsoleObject* ConsoleObjectPtr;
+	FDateTime TimeOfLastConsoleObjectRefresh;
+
+	double ConsoleObjectRefreshThreshold = 1.0;
+	
+	/** The value of this variable (if Variable object type) when the module started in this session after it may have been set by an ini file. */
+	FString StartupValueAsString;
+	
+	/** The source of this variable's (if Variable object type) last setting as recorded when the plugin was loaded. */
+	EConsoleVariableFlags StartupSource = ECVF_Default;
+
+	/** When variables change, this callback is executed. */
+	FDelegateHandle OnVariableChangedCallbackHandle;
+
+	/** When commands are unregistered change, this callback is broadcasted. */
+	FOnDetectConsoleObjectUnregistered OnDetectConsoleObjectUnregistered;
+	FDelegateHandle OnDetectConsoleObjectUnregisteredHandle;
+
+	/** A mapping of SetBy console variable flags to information like the associated display text. */
+	static const inline TArray<FStaticConsoleVariableFlagInfo> SupportedFlags =
+	{
+		{ EConsoleVariableFlags::ECVF_SetByConstructor, LOCTEXT("SetByConstructor", "Constructor") },
+		{ EConsoleVariableFlags::ECVF_SetByScalability, LOCTEXT("SetByScalability", "Scalability") },
+		{ EConsoleVariableFlags::ECVF_SetByGameSetting, LOCTEXT("SetByGameSetting", "Game Setting") },
+		{ EConsoleVariableFlags::ECVF_SetByProjectSetting, LOCTEXT("SetByProjectSetting", "Project Setting") },
+		{ EConsoleVariableFlags::ECVF_SetBySystemSettingsIni, LOCTEXT("SetBySystemSettingsIni", "System Settings ini") },
+		{ EConsoleVariableFlags::ECVF_SetByDeviceProfile, LOCTEXT("SetByDeviceProfile", "Device Profile") },
+		{ EConsoleVariableFlags::ECVF_SetByGameOverride, LOCTEXT("SetByGameOverride", "Game Override") },
+		{ EConsoleVariableFlags::ECVF_SetByConsoleVariablesIni, LOCTEXT("SetByConsoleVariablesIni", "Console Variables ini") },
+		{ EConsoleVariableFlags::ECVF_SetByCommandline, LOCTEXT("SetByCommandline", "Command line") },
+		{ EConsoleVariableFlags::ECVF_SetByCode, LOCTEXT("SetByCode", "Code") },
+		{ EConsoleVariableFlags::ECVF_SetByConsole, LOCTEXT("SetByConsole", "Console") }
+	};
+};
+
+#undef LOCTEXT_NAMESPACE
diff --git a/Engine/Plugins/Editor/ConsoleVariablesEditor/Source/ConsoleVariablesEditor/Private/ConsoleVariablesEditorFunctionLibrary.cpp b/Engine/Plugins/Editor/ConsoleVariablesEditor/Source/ConsoleVariablesEditor/Private/ConsoleVariablesEditorFunctionLibrary.cpp
new file mode 100644
--- /dev/null
+++ b/Engine/Plugins/Editor/ConsoleVariablesEditor/Source/ConsoleVariablesEditor/Private/ConsoleVariablesEditorFunctionLibrary.cpp
@@ -0,0 +1,216 @@
+// Copyright Epic Games, Inc. All Rights Reserved.
+
+#include "ConsoleVariablesEditorFunctionLibrary.h"
+
+#include "ConsoleVariablesAsset.h"
+#include "ConsoleVariablesEditorLog.h"
+#include "ConsoleVariablesEditorModule.h"
+#include "MultiUser/ConsoleVariableSyncData.h"
+#include "Widgets/Input/SCheckBox.h"
+
+UConsoleVariablesAsset* UConsoleVariablesEditorFunctionLibrary::GetCurrentlyLoadedPreset()
+{
+	const FConsoleVariablesEditorModule& ConsoleVariablesEditorModule = FConsoleVariablesEditorModule::Get();
+	
+	return ConsoleVariablesEditorModule.GetPresetAsset();
+}
+
+void UConsoleVariablesEditorFunctionLibrary::LoadPresetIntoConsoleVariablesEditor(const UConsoleVariablesAsset* InAsset)
+{
+	if (!InAsset)
+	{
+		UE_LOG(LogConsoleVariablesEditor, Error,
+			TEXT("%hs: InAsset was nullptr, please verify that InAsset is valid."),
+			__FUNCTION__);
+		return;
+	}
+	
+	const FConsoleVariablesEditorModule& ConsoleVariablesEditorModule = FConsoleVariablesEditorModule::Get();
+
+	ConsoleVariablesEditorModule.OpenConsoleVariablesDialogWithAssetSelected(InAsset);
+}
+
+bool UConsoleVariablesEditorFunctionLibrary::CopyCurrentListToAsset(UConsoleVariablesAsset* InAsset)
+{
+	if (!InAsset)
+	{
+		UE_LOG(LogConsoleVariablesEditor, Error,
+			TEXT("%hs: InAsset was nullptr, please verify that InAsset is valid."),
+			__FUNCTION__);
+		return false;
+	}
+	
+	const FConsoleVariablesEditorModule& ConsoleVariablesEditorModule = FConsoleVariablesEditorModule::Get();
+
+	InAsset->CopyFrom(ConsoleVariablesEditorModule.GetPresetAsset());
+	InAsset->MarkPackageDirty();
+	return true;
+}
+
+bool UConsoleVariablesEditorFunctionLibrary::AddValidatedCommandToCurrentPreset(const FString NewCommand)
+{
+	const FConsoleVariablesEditorModule& ConsoleVariablesEditorModule = FConsoleVariablesEditorModule::Get();
+
+	if (const FReply Reply = ConsoleVariablesEditorModule.ValidateConsoleInputAndAddToCurrentPreset(FText::FromString(NewCommand));
+		Reply.IsEventHandled())
+	{
+		GetCurrentlyLoadedPreset()->MarkPackageDirty();
+
+		return true;
+	}
+
+	return false;
+}
+
+bool UConsoleVariablesEditorFunctionLibrary::RemoveCommandFromCurrentPreset(const FString NewCommand)
+{
+	const FConsoleVariablesEditorModule& ConsoleVariablesEditorModule = FConsoleVariablesEditorModule::Get();
+	
+	if (const bool bSuccess = GetCurrentlyLoadedPreset()->RemoveConsoleVariable(NewCommand))
+	{
+		ConsoleVariablesEditorModule.RebuildList();
+		
+		GetCurrentlyLoadedPreset()->MarkPackageDirty();
+
+		return bSuccess;
+	}
+
+	return false;
+}
+
+bool UConsoleVariablesEditorFunctionLibrary::GetListOfCommandsFromPreset(
+ 	const UConsoleVariablesAsset* InAsset, TArray<FString>& OutCommandList)
+{
+	if (!InAsset)
+	{
+		UE_LOG(LogConsoleVariablesEditor, Error,
+			TEXT("%hs: InAsset was nullptr, please verify that InAsset is valid."),
+			__FUNCTION__);
+		return false;
+	}
+
+	OutCommandList.Empty();
+	for (const FConsoleVariablesEditorAssetSaveData& Command : InAsset->GetSavedCommands())
+	{
+		OutCommandList.Add(Command.CommandName);
+	}
+
+	if (OutCommandList.Num() < 0)
+	{
+		UE_LOG(LogConsoleVariablesEditor, Error,
+			TEXT("%hs: InAsset has no saved commands. Please verify you're using the right asset."),
+			__FUNCTION__);
+		return false;
+	}
+	
+	return true;
+}
+
+bool UConsoleVariablesEditorFunctionLibrary::SetConsoleVariableByName_Float(const FString InCommandName,
+	const float InValue)
+{
+	return SetConsoleVariableByName_String(InCommandName, FString::SanitizeFloat(InValue));
+}
+
+bool UConsoleVariablesEditorFunctionLibrary::SetConsoleVariableByName_Bool(const FString InCommandName,
+	const bool InValue)
+{
+	return SetConsoleVariableByName_String(InCommandName, InValue ? "true" : "false");
+}
+
+bool UConsoleVariablesEditorFunctionLibrary::SetConsoleVariableByName_Int(const FString InCommandName,
+	const int32 InValue)
+{
+	return SetConsoleVariableByName_String(InCommandName, FString::FromInt(InValue));
+}
+
+bool UConsoleVariablesEditorFunctionLibrary::SetConsoleVariableByName_String(const FString InCommandName,
+	const FString InValue)
+{
+	FConsoleVariablesEditorModule& ConsoleVariablesEditorModule = FConsoleVariablesEditorModule::Get();
+
+	if (const TWeakPtr<FConsoleVariablesEditorCommandInfo> CommandInfo =
+		ConsoleVariablesEditorModule.FindCommandInfoByName(InCommandName); CommandInfo.IsValid())
+	{
+		if (CommandInfo.Pin()->GetConsoleVariablePtr())
+		{
+			CommandInfo.Pin()->ExecuteCommand(InValue);
+			return true;
+		}
+		
+		UE_LOG(LogConsoleVariablesEditor, Error,
+			TEXT("%hs: Command %s is not a variable type. Please do not enter console commands, only console variables."),
+			__FUNCTION__, *InCommandName);
+	}
+	else
+	{
+		UE_LOG(LogConsoleVariablesEditor, Error,
+			TEXT("%hs: FConsoleVariablesEditorCommandInfo was not foundwith given name: %s."),
+			__FUNCTION__, *InCommandName);
+	}
+	return false;
+}
+
+bool UConsoleVariablesEditorFunctionLibrary::GetConsoleVariableStringValue(
+	const FString InCommandName, FString& OutValue)
+{
+	FConsoleVariablesEditorModule& ConsoleVariablesEditorModule = FConsoleVariablesEditorModule::Get();
+
+	if (const TWeakPtr<FConsoleVariablesEditorCommandInfo> CommandInfo =
+		ConsoleVariablesEditorModule.FindCommandInfoByName(InCommandName); CommandInfo.IsValid())
+	{
+		if (const IConsoleVariable* AsVariable = CommandInfo.Pin()->GetConsoleVariablePtr())
+		{
+			OutValue = AsVariable->GetString();
+			return true;
+		}
+	}
+
+	UE_LOG(LogConsoleVariablesEditor, Error,
+		TEXT("%hs: FConsoleVariablesEditorCommandInfo was not foundwith given name: %s."),
+		__FUNCTION__, *InCommandName);
+	return false;
+}
+
+bool UConsoleVariablesEditorFunctionLibrary::GetConsoleVariableSourceByName(const FString InCommandName, FString& OutValue)
+{
+	FConsoleVariablesEditorModule& ConsoleVariablesEditorModule = FConsoleVariablesEditorModule::Get();
+
+	if (const TWeakPtr<FConsoleVariablesEditorCommandInfo> CommandInfo =
+		ConsoleVariablesEditorModule.FindCommandInfoByName(InCommandName); CommandInfo.IsValid())
+	{
+		OutValue = CommandInfo.Pin()->GetSourceAsText().ToString();
+		return true;
+	}
+
+	UE_LOG(LogConsoleVariablesEditor, Error,
+		TEXT("%hs: FConsoleVariablesEditorCommandInfo was not found with given name: %s."),
+		__FUNCTION__, *InCommandName);
+	return false;
+}
+
+bool UConsoleVariablesEditorFunctionLibrary::GetEnableMultiUserCVarSync()
+{
+	if (const UConcertCVarSynchronization* CVarSync = GetDefault<UConcertCVarSynchronization>())
+	{
+		return CVarSync->bSyncCVarTransactions;
+	}
+
+	UE_LOG(LogConsoleVariablesEditor, Error,
+		TEXT("%hs: UConcertCVarSynchronization was not found, please verify the Multi-user plugin is enabled and this object class is compiled."),
+		__FUNCTION__);
+	return false;
+}
+
+void UConsoleVariablesEditorFunctionLibrary::SetEnableMultiUserCVarSync(bool bNewSetting)
+{
+	if (UConcertCVarSynchronization* CVarSync = GetMutableDefault<UConcertCVarSynchronization>())
+	{
+		CVarSync->bSyncCVarTransactions = bNewSetting;
+		return;
+	}
+
+	UE_LOG(LogConsoleVariablesEditor, Error,
+		TEXT("%hs: UConcertCVarSynchronization was not found, please verify the Multi-user plugin is enabled and this object class is compiled."),
+		__FUNCTION__);
+}
diff --git a/Engine/Plugins/Editor/ConsoleVariablesEditor/Source/ConsoleVariablesEditor/Private/ConsoleVariablesEditorModule.cpp b/Engine/Plugins/Editor/ConsoleVariablesEditor/Source/ConsoleVariablesEditor/Private/ConsoleVariablesEditorModule.cpp
new file mode 100644
--- /dev/null
+++ b/Engine/Plugins/Editor/ConsoleVariablesEditor/Source/ConsoleVariablesEditor/Private/ConsoleVariablesEditorModule.cpp
@@ -0,0 +1,452 @@
+// Copyright Epic Games, Inc. All Rights Reserved.
+
+#include "ConsoleVariablesEditorModule.h"
+
+#include "AssetTypeActions/AssetTypeActions_ConsoleVariables.h"
+#include "ConsoleVariablesEditorLog.h"
+#include "ConsoleVariablesEditorProjectSettings.h"
+#include "ConsoleVariablesEditorStyle.h"
+#include "MultiUser/ConsoleVariableSyncData.h"
+#include "Views/MainPanel/ConsoleVariablesEditorMainPanel.h"
+
+#include "Algo/AllOf.h"
+#include "Algo/Find.h"
+#include "Framework/Docking/TabManager.h"
+#include "ISettingsModule.h"
+#include "LevelEditor.h"
+#include "ToolMenus.h"
+#include "WorkspaceMenuStructure.h"
+#include "WorkspaceMenuStructureModule.h"
+
+#define LOCTEXT_NAMESPACE "FConsoleVariablesEditorModule"
+
+const FName FConsoleVariablesEditorModule::ConsoleVariablesToolkitPanelTabId(TEXT("ConsoleVariablesToolkitPanel"));
+
+FConsoleVariablesEditorModule& FConsoleVariablesEditorModule::Get()
+{
+	return FModuleManager::GetModuleChecked<FConsoleVariablesEditorModule>("ConsoleVariablesEditor");
+}
+
+void FConsoleVariablesEditorModule::StartupModule()
+{
+	IAssetTools& AssetTools = FModuleManager::LoadModuleChecked<FAssetToolsModule>("AssetTools").Get();
+
+	AssetTools.RegisterAssetTypeActions(MakeShared<FAssetTypeActions_ConsoleVariables>());
+
+	FConsoleVariablesEditorStyle::Initialize();
+
+	FCoreDelegates::OnFEngineLoopInitComplete.AddRaw(this, &FConsoleVariablesEditorModule::OnFEngineLoopInitComplete);
+}
+
+void FConsoleVariablesEditorModule::ShutdownModule()
+{
+	UToolMenus::UnregisterOwner(this);
+
+	FCoreDelegates::OnFEngineLoopInitComplete.RemoveAll(this);
+
+	FConsoleVariablesEditorStyle::Shutdown();
+	
+	MainPanel.Reset();
+
+	ConsoleObjectsMasterReference.Empty();
+
+	// Unregister project settings
+	ISettingsModule& SettingsModule = FModuleManager::LoadModuleChecked<ISettingsModule>("Settings");
+	{
+		SettingsModule.UnregisterSettings("Project", "Plugins", "Console Variables Editor");
+	}
+}
+
+void FConsoleVariablesEditorModule::SavePreset() const
+{
+	if (MainPanel.IsValid())
+	{
+		MainPanel->SaveCurrentPreset();
+	}
+}
+
+void FConsoleVariablesEditorModule::SaveSpecificPreset(const TObjectPtr<UConsoleVariablesAsset> Preset) const
+{
+	if (MainPanel.IsValid())
+	{
+		MainPanel->SaveSpecificPreset(Preset);
+	}
+}
+
+void FConsoleVariablesEditorModule::SavePresetAs() const
+{
+	if (MainPanel.IsValid())
+	{
+		MainPanel->SaveCurrentPresetAs();
+	}
+}
+
+void FConsoleVariablesEditorModule::SaveSpecificPresetAs(const TObjectPtr<UConsoleVariablesAsset> Preset) const
+{
+	if (MainPanel.IsValid())
+	{
+		MainPanel->SaveSpecificPresetAs(Preset);
+	}
+}
+
+void FConsoleVariablesEditorModule::OpenConsoleVariablesDialogWithPreset(
+	const TObjectPtr<UConsoleVariablesAsset> Preset) const
+{
+	if (MainPanel.IsValid())
+	{
+		MainPanel->ImportPreset(Preset);
+	}
+}
+
+void FConsoleVariablesEditorModule::OpenConsoleVariablesDialogWithAssetSelected(const FAssetData& InAssetData) const
+{
+	if (InAssetData.IsValid())
+	{
+		OpenConsoleVariablesEditor();
+	}
+
+	if (MainPanel.IsValid())
+	{
+		MainPanel->ImportPreset(InAssetData);
+	}
+}
+
+void FConsoleVariablesEditorModule::QueryAndBeginTrackingConsoleVariables()
+{
+	const int32 VariableCount = ConsoleObjectsMasterReference.Num();
+	
+	ConsoleObjectsMasterReference.Empty(VariableCount);
+	
+	IConsoleManager::Get().ForEachConsoleObjectThatStartsWith(FConsoleObjectVisitor::CreateLambda(
+		[this] (const TCHAR* Key, IConsoleObject* ConsoleObject)
+		{
+			if (!ConsoleObject || ConsoleObject->TestFlags(ECVF_Unregistered))
+			{
+				return;
+			}
+
+			const TSharedRef<FConsoleVariablesEditorCommandInfo> Info =
+				MakeShared<FConsoleVariablesEditorCommandInfo>(Key);
+			
+			Info->StartupSource = Info->GetSource();
+			Info->OnDetectConsoleObjectUnregisteredHandle = Info->OnDetectConsoleObjectUnregistered.AddRaw(
+				this, &FConsoleVariablesEditorModule::OnDetectConsoleObjectUnregistered);
+
+			if (IConsoleVariable* AsVariable = ConsoleObject->AsVariable())
+			{
+				Info->OnVariableChangedCallbackHandle = AsVariable->OnChangedDelegate().AddRaw(
+					this, &FConsoleVariablesEditorModule::OnConsoleVariableChanged);
+			}
+			
+			AddConsoleObjectCommandInfoToMasterReference(Info);
+		}),
+		TEXT(""));
+}
+
+TWeakPtr<FConsoleVariablesEditorCommandInfo> FConsoleVariablesEditorModule::FindCommandInfoByName(const FString& NameToSearch, ESearchCase::Type InSearchCase)
+{
+	TSharedPtr<FConsoleVariablesEditorCommandInfo>* Match = Algo::FindByPredicate(
+		ConsoleObjectsMasterReference,
+		[&NameToSearch, InSearchCase](const TSharedPtr<FConsoleVariablesEditorCommandInfo> Comparator)
+		{
+			return Comparator->Command.Equals(NameToSearch, InSearchCase);
+		});
+
+	return Match ? *Match : nullptr;
+}
+
+TArray<TWeakPtr<FConsoleVariablesEditorCommandInfo>> FConsoleVariablesEditorModule::FindCommandInfosMatchingTokens(
+	const TArray<FString>& InTokens, ESearchCase::Type InSearchCase)
+{
+	TArray<TWeakPtr<FConsoleVariablesEditorCommandInfo>> ReturnValue;
+	
+	for (const TSharedPtr<FConsoleVariablesEditorCommandInfo>& CommandInfo : ConsoleObjectsMasterReference)
+	{
+		// Match any
+		for (const FString& Token : InTokens)
+		{
+			bool bMatchFound = false;
+			
+			// Match all of these
+			const FString SpaceDelimiter = " ";
+			TArray<FString> OutSpacedArray;
+			if (Token.Contains(SpaceDelimiter) && Token.ParseIntoArray(OutSpacedArray, *SpaceDelimiter, true) > 1)
+			{
+				bMatchFound = Algo::AllOf(OutSpacedArray, [&CommandInfo, InSearchCase](const FString& Comparator)
+				{
+					return CommandInfo->Command.Contains(Comparator, InSearchCase);
+				});
+			}
+			else
+			{
+				bMatchFound = CommandInfo->Command.Contains(Token, InSearchCase);
+			}
+
+			if (bMatchFound)
+			{
+				ReturnValue.Add(CommandInfo);
+			}
+		}
+	}
+
+	return ReturnValue;
+}
+
+TWeakPtr<FConsoleVariablesEditorCommandInfo> FConsoleVariablesEditorModule::FindCommandInfoByConsoleObjectReference(
+	IConsoleObject* InConsoleObjectReference)
+{
+	TSharedPtr<FConsoleVariablesEditorCommandInfo>* Match = Algo::FindByPredicate(
+	ConsoleObjectsMasterReference,
+	[InConsoleObjectReference](const TSharedPtr<FConsoleVariablesEditorCommandInfo> Comparator)
+	{
+		return Comparator->GetConsoleObjectPtr() == InConsoleObjectReference;
+	});
+
+	return Match ? *Match : nullptr;
+}
+
+TObjectPtr<UConsoleVariablesAsset> FConsoleVariablesEditorModule::GetPresetAsset() const
+{
+	return EditingPresetAsset;
+}
+
+TObjectPtr<UConsoleVariablesAsset> FConsoleVariablesEditorModule::GetGlobalSearchAsset() const
+{
+	return EditingGlobalSearchAsset;
+}
+
+FReply FConsoleVariablesEditorModule::ValidateConsoleInputAndAddToCurrentPreset(const FText& CommittedText) const
+{
+	return MainPanel->ValidateConsoleInputAndAddToCurrentPreset(CommittedText);
+}
+
+void FConsoleVariablesEditorModule::RebuildList(const FString InConsoleCommandToScrollTo, bool bShouldCacheValues) const
+{
+	MainPanel->RebuildList(InConsoleCommandToScrollTo, bShouldCacheValues);
+}
+
+void FConsoleVariablesEditorModule::RefreshList() const
+{
+	MainPanel->RefreshList();
+}
+
+void FConsoleVariablesEditorModule::UpdatePresetValuesForSave(TObjectPtr<UConsoleVariablesAsset> InAsset)
+{
+	MainPanel->UpdatePresetValuesForSave(InAsset);
+}
+
+bool FConsoleVariablesEditorModule::PopulateGlobalSearchAssetWithVariablesMatchingTokens(const TArray<FString>& InTokens)
+{
+	// Remove existing commands
+	TArray<FConsoleVariablesEditorAssetSaveData> NullList;
+	EditingGlobalSearchAsset->ReplaceSavedCommands(NullList);
+	
+	for (const TWeakPtr<FConsoleVariablesEditorCommandInfo>& CommandInfo : FindCommandInfosMatchingTokens(InTokens))
+	{
+		EditingGlobalSearchAsset->AddOrSetConsoleObjectSavedData(
+			{
+				CommandInfo.Pin()->Command,
+				CommandInfo.Pin()->GetConsoleVariablePtr() ?
+					CommandInfo.Pin()->GetConsoleVariablePtr()->GetString() : "",
+				ECheckBoxState::Checked
+			}
+		);
+	}
+
+	return EditingGlobalSearchAsset->GetSavedCommandsCount() > 0;
+}
+
+void FConsoleVariablesEditorModule::SendMultiUserConsoleVariableChange(const FString& InVariableName, const FString& InValueAsString) const
+{
+	MainPanel->GetMultiUserManager().SendConsoleVariableChange(InVariableName, InValueAsString);
+}
+
+void FConsoleVariablesEditorModule::OnRemoteCvarChanged(const FString InName, const FString InValue)
+{
+	UE_LOG(LogConsoleVariablesEditor, Display, TEXT("Remote set console variable %s = %s"), *InName, *InValue);
+
+	if (GetDefault<UConcertCVarSynchronization>()->bSyncCVarTransactions)
+	{
+		bool bShouldExecute = true;
+		
+		if (const TWeakPtr<FConsoleVariablesEditorCommandInfo> CommandInfo =
+			FindCommandInfoByName(InName); CommandInfo.IsValid())
+		{
+			bShouldExecute = CommandInfo.Pin()->IsCurrentValueDifferentFromInputValue(InValue);
+		}
+
+		if (bShouldExecute)
+		{
+			GEngine->Exec(FConsoleVariablesEditorCommandInfo::GetCurrentWorld(),
+				*FString::Printf(TEXT("%s %s"), *InName, *InValue));
+			
+			EditingPresetAsset->AddOrSetConsoleObjectSavedData(
+				{
+					InName,
+					InValue,
+					ECheckBoxState::Checked
+				}
+			);
+
+			if (MainPanel->GetEditorListMode() == FConsoleVariablesEditorList::EConsoleVariablesEditorListMode::Preset)
+			{
+				MainPanel->RebuildList();
+			}
+		}
+	}
+}
+
+void FConsoleVariablesEditorModule::OnFEngineLoopInitComplete()
+{
+	RegisterMenuItem();
+	RegisterProjectSettings();
+	QueryAndBeginTrackingConsoleVariables();
+	CreateEditingPresets();
+	
+	MainPanel = MakeShared<FConsoleVariablesEditorMainPanel>();
+}
+
+void FConsoleVariablesEditorModule::RegisterMenuItem()
+{
+	FTabSpawnerEntry& BrowserSpawnerEntry = FGlobalTabmanager::Get()->RegisterNomadTabSpawner(
+		ConsoleVariablesToolkitPanelTabId,
+		FOnSpawnTab::CreateRaw(this, & FConsoleVariablesEditorModule::SpawnMainPanelTab))
+			.SetIcon(FSlateIcon(FConsoleVariablesEditorStyle::Get().GetStyleSetName(), "ConsoleVariables.ToolbarButton", "ConsoleVariables.ToolbarButton.Small"))
+			.SetDisplayName(LOCTEXT("OpenConsoleVariablesEditorMenuItem", "Console Variables"))
+			.SetTooltipText(LOCTEXT("OpenConsoleVariablesEditorTooltip", "Open the Console Variables Editor"))
+			.SetMenuType(ETabSpawnerMenuType::Enabled);
+
+	BrowserSpawnerEntry.SetGroup(WorkspaceMenu::GetMenuStructure().GetLevelEditorCategory());
+}
+
+void FConsoleVariablesEditorModule::RegisterProjectSettings() const
+{
+	ISettingsModule& SettingsModule = FModuleManager::LoadModuleChecked<ISettingsModule>("Settings");
+	{
+		// User Project Settings
+		const TSharedPtr<ISettingsSection> ProjectSettingsSectionPtr = SettingsModule.RegisterSettings(
+			"Project", "Plugins", "Console Variables Editor",
+			NSLOCTEXT("ConsoleVariables", "ConsoleVariablesSettingsCategoryDisplayName", "Console Variables Editor"),
+			NSLOCTEXT("ConsoleVariables", "ConsoleVariablesSettingsDescription",
+			          "Configure the Console Variables Editor user settings"),
+			GetMutableDefault<UConsoleVariablesEditorProjectSettings>());
+	}
+}
+
+void FConsoleVariablesEditorModule::OnConsoleVariableChanged(IConsoleVariable* ChangedVariable)
+{
+	check(EditingPresetAsset);
+
+	if (const TWeakPtr<FConsoleVariablesEditorCommandInfo> CommandInfo =
+		FindCommandInfoByConsoleObjectReference(ChangedVariable); CommandInfo.IsValid())
+	{
+		const TSharedPtr<FConsoleVariablesEditorCommandInfo>& PinnedCommand = CommandInfo.Pin();
+		const FString& Key = PinnedCommand->Command;
+		
+		FConsoleVariablesEditorAssetSaveData FoundData;
+		const bool bIsVariableCurrentlyTracked = EditingPresetAsset->FindSavedDataByCommandString(Key, FoundData);
+
+		const UConsoleVariablesEditorProjectSettings* Settings = GetDefault<UConsoleVariablesEditorProjectSettings>();
+		check(Settings);
+		
+		if (!bIsVariableCurrentlyTracked)
+		{
+			// If not yet tracked and we want to track variable changes from outside the dialogue,
+			// Check if the changed value differs from the startup value before tracking it
+			if (Settings->bAddAllChangedConsoleVariablesToCurrentPreset &&
+				!Settings->ChangedConsoleVariableSkipList.Contains(Key) && 
+				PinnedCommand->IsCurrentValueDifferentFromInputValue(PinnedCommand->StartupValueAsString))
+			{
+				if (MainPanel.IsValid())
+				{
+					MainPanel->AddConsoleObjectToCurrentPreset(
+						Key,
+						// If we're not in preset mode then pass empty value
+						// This forces the row to get the current value at the time it's generated
+						MainPanel->GetEditorListMode() ==
+						FConsoleVariablesEditorList::EConsoleVariablesEditorListMode::Preset ?
+						ChangedVariable->GetString() : "",
+						true
+					);
+				}
+
+				SendMultiUserConsoleVariableChange(Key, ChangedVariable->GetString());
+			}
+		}
+		else // If it's already being tracked, refreshed the list to update show filters and other possibly stale elements
+		{
+			if (MainPanel.IsValid())
+			{
+				MainPanel->RefreshList();
+			}
+			
+			SendMultiUserConsoleVariableChange(Key, ChangedVariable->GetString());
+		}
+	}
+}
+
+void FConsoleVariablesEditorModule::OnDetectConsoleObjectUnregistered(FString CommandName)
+{
+	check(EditingPresetAsset);
+
+	EditingPresetAsset->RemoveConsoleVariable(CommandName);
+
+	if (MainPanel.IsValid())
+	{
+		MainPanel->RefreshList();
+	}
+
+	if (const TWeakPtr<FConsoleVariablesEditorCommandInfo> CommandInfo =
+		FindCommandInfoByName(CommandName); CommandInfo.IsValid())
+	{
+		ConsoleObjectsMasterReference.Remove(CommandInfo.Pin());
+	}
+}
+
+TObjectPtr<UConsoleVariablesAsset> FConsoleVariablesEditorModule::AllocateTransientPreset(const FName DesiredName) const
+{
+	const FString PackageName = "/Temp/ConsoleVariablesEditor/PendingConsoleVariablesPresets";
+
+	UPackage* NewPackage = CreatePackage(*PackageName);
+	NewPackage->SetFlags(RF_Transient);
+	NewPackage->AddToRoot();
+
+	return NewObject<UConsoleVariablesAsset>(
+		NewPackage, DesiredName, RF_Transient | RF_Transactional | RF_Standalone);
+}
+
+void FConsoleVariablesEditorModule::CreateEditingPresets()
+{
+	EditingPresetAsset = AllocateTransientPreset("ConsoleVariablesPreset_PendingPreset");
+	
+	EditingGlobalSearchAsset = AllocateTransientPreset("ConsoleVariablesPreset_GlobalSearch");
+}
+
+TSharedRef<SDockTab> FConsoleVariablesEditorModule::SpawnMainPanelTab(const FSpawnTabArgs& Args)
+{
+	const TSharedRef<SDockTab> DockTab = SNew(SDockTab).TabRole(ETabRole::NomadTab);
+	DockTab->SetContent(MainPanel->GetOrCreateWidget());
+	MainPanel->RebuildList();
+			
+	return DockTab;
+}
+
+void FConsoleVariablesEditorModule::OpenConsoleVariablesEditor()
+{
+	FGlobalTabmanager::Get()->TryInvokeTab(ConsoleVariablesToolkitPanelTabId);
+}
+
+void FConsoleVariablesEditorModule::AddReferencedObjects(FReferenceCollector& Collector)
+{
+	Collector.AddReferencedObject(EditingPresetAsset);
+	Collector.AddReferencedObject(EditingGlobalSearchAsset);
+}
+
+FString FConsoleVariablesEditorModule::GetReferencerName() const
+{
+	return TEXT("FConsoleVariablesEditorModule");
+}
+
+#undef LOCTEXT_NAMESPACE
+
+IMPLEMENT_MODULE(FConsoleVariablesEditorModule, ConsoleVariablesEditor)
diff --git a/Engine/Plugins/Editor/ConsoleVariablesEditor/Source/ConsoleVariablesEditor/Private/ConsoleVariablesEditorModule.h b/Engine/Plugins/Editor/ConsoleVariablesEditor/Source/ConsoleVariablesEditor/Private/ConsoleVariablesEditorModule.h
new file mode 100644
--- /dev/null
+++ b/Engine/Plugins/Editor/ConsoleVariablesEditor/Source/ConsoleVariablesEditor/Private/ConsoleVariablesEditorModule.h
@@ -0,0 +1,123 @@
+// Copyright Epic Games, Inc. All Rights Reserved.
+
+#pragma once
+
+#include "ConsoleVariablesAsset.h"
+#include "ConsoleVariablesEditorCommandInfo.h"
+#include "MultiUser/ConsoleVariableSync.h"
+
+#include "CoreMinimal.h"
+#include "HAL/IConsoleManager.h"
+#include "Modules/ModuleInterface.h"
+#include "Widgets/Docking/SDockTab.h"
+
+struct FAssetData;
+
+class FConsoleVariablesEditorMainPanel;
+class FConsoleVariablesEditorToolkit;
+class ISettingsSection;
+
+class FConsoleVariablesEditorModule : public IModuleInterface, public FGCObject
+{
+public:
+	
+	static FConsoleVariablesEditorModule& Get();
+
+	//~ Begin IModuleInterface Interface
+	virtual void StartupModule() override;
+	virtual void ShutdownModule() override;
+	//~ End IModuleInterface Interface
+
+	void SavePreset() const;
+	void SaveSpecificPreset(const TObjectPtr<UConsoleVariablesAsset> Preset) const;
+	void SavePresetAs() const;
+	void SaveSpecificPresetAs(const TObjectPtr<UConsoleVariablesAsset> Preset) const;
+	void OpenConsoleVariablesDialogWithPreset(const TObjectPtr<UConsoleVariablesAsset> Preset) const;
+	void OpenConsoleVariablesDialogWithAssetSelected(const FAssetData& InAssetData) const;
+
+	/** Find all console variables and cache their startup values */
+	void QueryAndBeginTrackingConsoleVariables();
+
+	void AddConsoleObjectCommandInfoToMasterReference(TSharedRef<FConsoleVariablesEditorCommandInfo> InCommandInfo)
+	{
+		ConsoleObjectsMasterReference.Add(InCommandInfo);
+	}
+
+	/** Find a tracked console variable by the command string with optional case sensitivity. */
+	TWeakPtr<FConsoleVariablesEditorCommandInfo> FindCommandInfoByName(
+		const FString& NameToSearch, ESearchCase::Type InSearchCase = ESearchCase::IgnoreCase);
+
+	/*Find all tracked console variables matching a specific search query with optional case sensitivity.
+	 *Individual members of InTokens will be considered "AnyOf" or "OR" searches. If SearchTerms contains any individual member it will match.
+	 *Members will be tested for a space character (" "). If a space is found, a subsearch will be run.
+	 *This subsearch will be an "AllOf" or "AND" type search in which all strings, separated by a space, must be found in the search terms.
+	 */
+	TArray<TWeakPtr<FConsoleVariablesEditorCommandInfo>> FindCommandInfosMatchingTokens(
+		const TArray<FString>& InTokens, ESearchCase::Type InSearchCase = ESearchCase::IgnoreCase);
+
+	/**
+	 *Find a tracked console variable by its console object reference.
+	 *Note that some commands do not have an associated console object (such as 'stat unit') and will not be found with this method.
+	 *It's normally safer to use FindCommandInfoByName() instead.
+	 */
+	TWeakPtr<FConsoleVariablesEditorCommandInfo> FindCommandInfoByConsoleObjectReference(
+		IConsoleObject* InConsoleObjectReference);
+	
+	[[nodiscard]] TObjectPtr<UConsoleVariablesAsset> GetPresetAsset() const;
+	[[nodiscard]] TObjectPtr<UConsoleVariablesAsset> GetGlobalSearchAsset() const;
+	
+	FReply ValidateConsoleInputAndAddToCurrentPreset(const FText& CommittedText) const;
+	/*
+	 * Regenerate the list items and refresh the list. Call when adding or removing variables.
+	 * @param bShouldCacheValues If true, the current list's current values will be cached and then restored when the list is rebuilt. Otherwise preset values will be used.
+	 */
+	void RebuildList(const FString InConsoleCommandToScrollTo = "", bool bShouldCacheValues = true) const;
+	
+	/**
+	 * Refresh filters and sorting.
+	 * Useful for when the list state has gone stale but the variable count has not changed.
+	 */
+	void RefreshList() const;
+	
+	void UpdatePresetValuesForSave(TObjectPtr<UConsoleVariablesAsset> InAsset);
+
+	/** Fills Global Search Asset's Saved Commands with variables matching the specified query. Returns false if no matches were found. */
+	bool PopulateGlobalSearchAssetWithVariablesMatchingTokens(const TArray<FString>& InTokens);
+
+	void SendMultiUserConsoleVariableChange(const FString& InVariableName, const FString& InValueAsString) const;
+	void OnRemoteCvarChanged(const FString InName, const FString InValue);
+
+	virtual void AddReferencedObjects( FReferenceCollector& Collector )  override;
+	virtual FString GetReferencerName() const override;
+
+private:
+
+	void OnFEngineLoopInitComplete();
+
+	void RegisterMenuItem();
+	void RegisterProjectSettings() const;
+
+	void OnConsoleVariableChanged(IConsoleVariable* ChangedVariable);
+	/** In the event a console object is unregistered, this failsafe callback will clean up the associated list item and command info object. */
+	void OnDetectConsoleObjectUnregistered(FString CommandName);
+
+	TObjectPtr<UConsoleVariablesAsset> AllocateTransientPreset(const FName DesiredName) const;
+	void CreateEditingPresets();
+
+	TSharedRef<SDockTab> SpawnMainPanelTab(const FSpawnTabArgs& Args);
+
+	static void OpenConsoleVariablesEditor();
+	
+	static const FName ConsoleVariablesToolkitPanelTabId;
+
+	/* Lives for as long as the module is loaded. */
+	TSharedPtr<FConsoleVariablesEditorMainPanel> MainPanel;
+
+	/** Transient preset that's being edited so we don't affect the reference asset unless we save it */
+	TObjectPtr<UConsoleVariablesAsset> EditingPresetAsset = nullptr;
+	/** Transient preset that tracks variables that match the search criteria */
+	TObjectPtr<UConsoleVariablesAsset> EditingGlobalSearchAsset = nullptr;
+
+	/** All tracked variables and their default, startup, and current values */
+	TArray<TSharedPtr<FConsoleVariablesEditorCommandInfo>> ConsoleObjectsMasterReference;
+};
diff --git a/Engine/Plugins/Editor/ConsoleVariablesEditor/Source/ConsoleVariablesEditor/Private/ConsoleVariablesEditorStyle.cpp b/Engine/Plugins/Editor/ConsoleVariablesEditor/Source/ConsoleVariablesEditor/Private/ConsoleVariablesEditorStyle.cpp
new file mode 100644
--- /dev/null
+++ b/Engine/Plugins/Editor/ConsoleVariablesEditor/Source/ConsoleVariablesEditor/Private/ConsoleVariablesEditorStyle.cpp
@@ -0,0 +1,111 @@
+// Copyright Epic Games, Inc. All Rights Reserved.
+
+#include "ConsoleVariablesEditorStyle.h"
+
+#include "EditorStyleSet.h"
+#include "Framework/Application/SlateApplication.h"
+#include "Interfaces/IPluginManager.h"
+#include "Styling/CoreStyle.h"
+#include "Styling/SlateStyleRegistry.h"
+#include "Styling/SlateTypes.h"
+#include "Styling/StyleColors.h"
+
+TSharedPtr<FSlateStyleSet> FConsoleVariablesEditorStyle::StyleInstance = nullptr;
+
+void FConsoleVariablesEditorStyle::Initialize()
+{
+	if (!StyleInstance.IsValid())
+	{
+		StyleInstance = Create();
+		FSlateStyleRegistry::RegisterSlateStyle(*StyleInstance);
+	}
+}
+
+void FConsoleVariablesEditorStyle::Shutdown()
+{
+	FSlateStyleRegistry::UnRegisterSlateStyle(*StyleInstance);
+	ensure(StyleInstance.IsUnique());
+	StyleInstance.Reset();
+}
+
+void FConsoleVariablesEditorStyle::ReloadTextures()
+{
+	FSlateApplication::Get().GetRenderer()->ReloadTextureResources();
+}
+
+const ISlateStyle& FConsoleVariablesEditorStyle::Get()
+{
+	return *StyleInstance;
+}
+
+const FName& FConsoleVariablesEditorStyle::GetStyleSetName() const
+{
+	static FName ConsoleVariablesStyleSetName(TEXT("ConsoleVariablesEditor"));
+	return ConsoleVariablesStyleSetName;
+}
+
+const FSlateBrush* FConsoleVariablesEditorStyle::GetBrush(const FName PropertyName, const ANSICHAR* Specifier, const ISlateStyle* RequestingStyle) const
+{
+	return StyleInstance->GetBrush(PropertyName, Specifier);
+}
+
+#define IMAGE_BRUSH( RelativePath, ... ) FSlateImageBrush( Style->RootToContentDir( RelativePath, TEXT(".png") ), __VA_ARGS__ )
+#define BOX_BRUSH( RelativePath, ... ) FSlateBoxBrush( Style->RootToContentDir( RelativePath, TEXT(".png") ), __VA_ARGS__ )
+#define BORDER_BRUSH( RelativePath, ... ) FSlateBorderBrush( Style->RootToContentDir( RelativePath, TEXT(".png") ), __VA_ARGS__ )
+#define IMAGE_BRUSH_SVG( RelativePath, ... ) FSlateVectorImageBrush( Style->RootToContentDir(RelativePath, TEXT(".svg") ), __VA_ARGS__)
+#define IMAGE_PLUGIN_BRUSH_SVG( RelativePath, ... ) FSlateVectorImageBrush( FConsoleVariablesEditorStyle::InContent(RelativePath, ".svg"), __VA_ARGS__)
+
+const FVector2D Icon64x64(64.f, 64.f);
+const FVector2D Icon40x40(40.0f, 40.0f);
+const FVector2D Icon20x20(20.0f, 20.0f);
+const FVector2D Icon16x16(16.0f, 16.0f);
+const FVector2D Icon12x12(12.0f, 12.0f);
+const FVector2D Icon8x8(8.f, 8.f);
+
+FString FConsoleVariablesEditorStyle::InContent(const FString& RelativePath, const ANSICHAR* Extension)
+{
+	static FString ContentDir = IPluginManager::Get().FindPlugin(TEXT("ConcertSyncClient"))->GetContentDir();
+	return (ContentDir / RelativePath) + Extension;
+}
+
+TSharedRef< FSlateStyleSet > FConsoleVariablesEditorStyle::Create()
+{
+	TSharedRef<FSlateStyleSet> Style = MakeShared<FSlateStyleSet>("ConsoleVariablesEditor");
+
+	TSharedPtr<IPlugin> Plugin = IPluginManager::Get().FindPlugin(TEXT("ConsoleVariables"));
+	check(Plugin.IsValid());
+	if (Plugin.IsValid())
+	{
+		Style->SetContentRoot(FPaths::Combine(Plugin->GetBaseDir(), TEXT("Resources")));
+	}
+
+	Style->SetCoreContentRoot(FPaths::EngineContentDir() / TEXT("Slate"));
+
+	// Icons
+	Style->Set("ConsoleVariables.ToolbarButton", new IMAGE_BRUSH_SVG("Icons/ConsoleVariables", Icon40x40));
+	Style->Set("ConsoleVariables.ToolbarButton.Small", new IMAGE_BRUSH_SVG("Icons/ConsoleVariables", Icon20x20));
+	
+	Style->Set("ConsoleVariables.Favorite.Outline", new IMAGE_BRUSH_SVG("Icons/FavoriteOutline", Icon40x40));
+	Style->Set("ConsoleVariables.Favorite.Outline.Small", new IMAGE_BRUSH_SVG("Icons/FavoriteOutline", Icon16x16));
+	
+	Style->Set("ConsoleVariables.GlobalSearch", new IMAGE_BRUSH_SVG("Icons/SearchGlobal", Icon40x40));
+	Style->Set("ConsoleVariables.GlobalSearch.Small", new IMAGE_BRUSH_SVG("Icons/SearchGlobal", Icon16x16));
+
+	// Brush
+	Style->Set("ConsoleVariablesEditor.GroupBorder", new BOX_BRUSH("Common/DarkGroupBorder", FMargin(4.0f / 16.0f)));
+	Style->Set("ConsoleVariablesEditor.BrightBorder", new FSlateColorBrush(FStyleColors::Header));
+
+	// Border colors for Results view
+	Style->Set("ConsoleVariablesEditor.HeaderRowBorder", new FSlateColorBrush(FStyleColors::Black));
+	Style->Set("ConsoleVariablesEditor.CommandGroupBorder", new BOX_BRUSH("Common/DarkGroupBorder", FMargin(4.0f / 16.0f)));
+	Style->Set("ConsoleVariablesEditor.DefaultBorder", new FSlateColorBrush(FStyleColors::Transparent));
+
+	// Multi-user Tab/Menu icons
+	Style->Set("Concert.MultiUser", new IMAGE_PLUGIN_BRUSH_SVG("Icons/icon_MultiUser", Icon16x16));
+	
+	return Style;
+}
+
+#undef IMAGE_BRUSH
+#undef BOX_BRUSH
+#undef BORDER_BRUSH
diff --git a/Engine/Plugins/Editor/ConsoleVariablesEditor/Source/ConsoleVariablesEditor/Private/ConsoleVariablesEditorStyle.h b/Engine/Plugins/Editor/ConsoleVariablesEditor/Source/ConsoleVariablesEditor/Private/ConsoleVariablesEditorStyle.h
new file mode 100644
--- /dev/null
+++ b/Engine/Plugins/Editor/ConsoleVariablesEditor/Source/ConsoleVariablesEditor/Private/ConsoleVariablesEditorStyle.h
@@ -0,0 +1,36 @@
+// Copyright Epic Games, Inc. All Rights Reserved.
+
+#pragma once
+
+#include "Styling/SlateStyle.h"
+
+class FConsoleVariablesEditorStyle : FSlateStyleSet
+{
+public:
+
+	static void Initialize();
+
+	static void Shutdown();
+
+	static void ReloadTextures();
+
+	static const ISlateStyle& Get();
+
+	virtual const FName& GetStyleSetName() const override;
+
+	virtual const FSlateBrush* GetBrush(const FName PropertyName, const ANSICHAR* Specifier = nullptr, const ISlateStyle* RequestingStyle = nullptr) const override;
+
+	template< typename WidgetStyleType >
+	static const WidgetStyleType& GetWidgetStyle(FName PropertyName, const ANSICHAR* Specifier = nullptr)
+	{
+		return StyleInstance->GetWidgetStyle<WidgetStyleType>(PropertyName, Specifier);
+	}
+
+private:
+	
+	static FString InContent(const FString& RelativePath, const ANSICHAR* Extension);
+	
+	static TSharedRef<FSlateStyleSet> Create();
+
+	static TSharedPtr<FSlateStyleSet> StyleInstance;
+};
diff --git a/Engine/Plugins/Editor/ConsoleVariablesEditor/Source/ConsoleVariablesEditor/Private/Factories/ConsoleVariablesEditorFactory.cpp b/Engine/Plugins/Editor/ConsoleVariablesEditor/Source/ConsoleVariablesEditor/Private/Factories/ConsoleVariablesEditorFactory.cpp
new file mode 100644
--- /dev/null
+++ b/Engine/Plugins/Editor/ConsoleVariablesEditor/Source/ConsoleVariablesEditor/Private/Factories/ConsoleVariablesEditorFactory.cpp
@@ -0,0 +1,38 @@
+// Copyright Epic Games, Inc. All Rights Reserved.
+
+#include "Factories/ConsoleVariablesEditorFactory.h"
+
+#include "AssetToolsModule.h"
+
+#include "ConsoleVariablesAsset.h"
+
+#define LOCTEXT_NAMESPACE "ConsoleVariablesEditorFactory"
+
+UConsoleVariablesEditorFactory::UConsoleVariablesEditorFactory()
+{
+	SupportedClass = UConsoleVariablesAsset::StaticClass();
+
+	// This factory manufacture new objects from scratch.
+	bCreateNew = true;
+
+	// This factory will open the editor for each new object.
+	bEditAfterNew = true;
+};
+
+UObject* UConsoleVariablesEditorFactory::FactoryCreateNew(UClass* InClass, UObject* InParent, FName InName, EObjectFlags Flags, UObject* Context, FFeedbackContext* Warn)
+{
+	return NewObject<UConsoleVariablesEditorFactory>(InParent, InClass, InName, Flags);
+};
+
+
+bool UConsoleVariablesEditorFactory::ShouldShowInNewMenu() const
+{
+	return true;
+}
+uint32 UConsoleVariablesEditorFactory::GetMenuCategories() const
+{
+	IAssetTools& AssetTools = FModuleManager::LoadModuleChecked<FAssetToolsModule>("AssetTools").Get();
+	return AssetTools.RegisterAdvancedAssetCategory("Console Variables Editor", LOCTEXT("AssetCategoryName", "Console Variables Editor"));
+}
+
+#undef LOCTEXT_NAMESPACE
diff --git a/Engine/Plugins/Editor/ConsoleVariablesEditor/Source/ConsoleVariablesEditor/Private/Factories/ConsoleVariablesEditorFactory.h b/Engine/Plugins/Editor/ConsoleVariablesEditor/Source/ConsoleVariablesEditor/Private/Factories/ConsoleVariablesEditorFactory.h
new file mode 100644
--- /dev/null
+++ b/Engine/Plugins/Editor/ConsoleVariablesEditor/Source/ConsoleVariablesEditor/Private/Factories/ConsoleVariablesEditorFactory.h
@@ -0,0 +1,22 @@
+// Copyright Epic Games, Inc. All Rights Reserved.
+
+#pragma once
+
+#include "CoreMinimal.h"
+#include "Factories/Factory.h"
+
+#include "ConsoleVariablesEditorFactory.generated.h"
+
+UCLASS(MinimalAPI)
+class UConsoleVariablesEditorFactory : public UFactory
+{
+	GENERATED_BODY()
+
+public:
+
+	UConsoleVariablesEditorFactory();
+	
+	virtual UObject* FactoryCreateNew(UClass* InClass, UObject* InParent, FName InName, EObjectFlags Flags, UObject* Context, FFeedbackContext* Warn) override;
+	virtual bool ShouldShowInNewMenu() const override;
+	virtual uint32 GetMenuCategories() const override;
+};
\ No newline at end of file
diff --git a/Engine/Plugins/Editor/ConsoleVariablesEditor/Source/ConsoleVariablesEditor/Private/MultiUser/ConcertConsoleVariableSessionCustomization.cpp b/Engine/Plugins/Editor/ConsoleVariablesEditor/Source/ConsoleVariablesEditor/Private/MultiUser/ConcertConsoleVariableSessionCustomization.cpp
new file mode 100644
--- /dev/null
+++ b/Engine/Plugins/Editor/ConsoleVariablesEditor/Source/ConsoleVariablesEditor/Private/MultiUser/ConcertConsoleVariableSessionCustomization.cpp
@@ -0,0 +1,411 @@
+// Copyright Epic Games, Inc. All Rights Reserved.
+
+#include "ConcertConsoleVariableSessionCustomization.h"
+
+#include "Containers/ContainersFwd.h"
+#include "Delegates/Delegate.h"
+#include "Framework/Views/ITypedTableView.h"
+#include "IConcertClient.h"
+#include "IConcertSyncClient.h"
+#include "IConcertSession.h"
+#include "IConcertSessionHandler.h"
+#include "IConcertSyncClientModule.h"
+
+#include "Misc/AssertionMacros.h"
+#include "Oculus/LibOVRPlatform/LibOVRPlatform/include/OVR_LaunchBlockFlowResult.h"
+#include "Styling/SlateTypes.h"
+#include "Templates/SharedPointer.h"
+#include "Types/SlateEnums.h"
+#include "UObject/Script.h"
+#include "UObject/UObjectGlobals.h"
+#include "Widgets/SBoxPanel.h"
+#include "Widgets/SWidget.h"
+#include "Widgets/Layout/SBox.h"
+#include "Widgets/Text/STextBlock.h"
+#include "Widgets/Input/SHyperlink.h"
+#include "Widgets/Views/SListView.h"
+#include "Widgets/Views/STableRow.h"
+#include "Widgets/Input/SCheckBox.h"
+#include "EditorFontGlyphs.h"
+#include "EditorStyleSet.h"
+#include "DetailWidgetRow.h"
+
+
+#include "ConcertMessages.h"
+
+#include "MultiUser/ConsoleVariableSyncData.h"
+
+#include "Core/Public/Modules/ModuleManager.h"
+#include "ISettingsModule.h"
+
+#define LOCTEXT_NAMESPACE "ConcertConsoleVariableSessionCustomization"
+
+namespace ConsoleVariableSessionUI
+{
+	static const FName DisplayNameColumnName(TEXT("UserDisplayName"));
+	static const FName TransactColumnName(TEXT("Transact"));
+}
+
+class SConsoleVariableSessionRow : public SMultiColumnTableRow<TSharedPtr<FConcertCVarDetails>>
+{
+	SLATE_BEGIN_ARGS(SConsoleVariableSessionRow) {}
+	SLATE_END_ARGS()
+
+public:
+
+	void Construct(const FArguments& InArgs, TWeakPtr<IConcertSyncClient> InSyncClient,
+				   TWeakPtr<IConcertClientSession> InClientSession,
+				   TSharedPtr<FConcertCVarDetails> InClientInfo, const TSharedRef<STableViewBase>& InOwnerTableView)
+	{
+		SyncClient = MoveTemp(InSyncClient);
+		ClientSession = MoveTemp(InClientSession);
+		ClientSetting = MoveTemp(InClientInfo);
+		SMultiColumnTableRow<TSharedPtr<FConcertCVarDetails>>::Construct(FSuperRowType::FArguments(), InOwnerTableView);
+
+		SetToolTipText(MakeAttributeSP(this, &SConsoleVariableSessionRow::GetRowToolTip));
+	}
+
+public:
+	TSharedRef<SWidget> GenerateWidgetForColumn(const FName& ColumnName) override
+	{
+		using ThisType = SConsoleVariableSessionRow;
+		if (ColumnName == ConsoleVariableSessionUI::DisplayNameColumnName)
+		{
+			return SNew(SHorizontalBox)
+				+SHorizontalBox::Slot()
+				.AutoWidth()
+				.VAlign(VAlign_Center)
+				.Padding(FMargin(4.0f, 0.0f, 0.0f, 0.0f))
+				[
+					SNew(STextBlock)
+					.Font(this, &SConsoleVariableSessionRow::GetAvatarFont)
+					.ColorAndOpacity(this, &SConsoleVariableSessionRow::GetAvatarColor)
+					.Text(FEditorFontGlyphs::Square)
+				]
+
+				// The client display name.
+				+SHorizontalBox::Slot()
+				.VAlign(VAlign_Center)
+				[
+					SNew(SBorder)
+					.BorderImage(FEditorStyle::GetBrush("NoBorder"))
+					.ColorAndOpacity(FLinearColor(0.75f, 0.75f, 0.75f))
+					.Padding(FMargin(6.0f, 4.0f))
+					[
+						SNew(STextBlock)
+						.Font(FEditorStyle::GetFontStyle("BoldFont"))
+						.Text(GetDisplayName())
+					]
+				];
+		}
+		else
+		{
+			check(ColumnName == ConsoleVariableSessionUI::TransactColumnName);
+
+			return  SNew(SHorizontalBox)
+				+SHorizontalBox::Slot()
+				.HAlign(HAlign_Left)
+				[
+					SAssignNew(Transact,SCheckBox)
+					.IsChecked(this, &ThisType::GetTransactState)
+					.OnCheckStateChanged(this, &ThisType::SetTransactState)
+					.IsEnabled(this, &ThisType::CheckBoxEnabled)
+					.ToolTipText_Lambda([this](){return GetTransactToolTip();})
+				];
+		}
+	}
+
+	FText GetToolTipNoSync()
+	{
+		return FText::Format(LOCTEXT("NoSyncEnabled_ToolTip", "Synchronize Disabled on {0}"), GetDisplayName());
+	}
+
+
+	FText GetTransactToolTip()
+	{
+		TSharedPtr<FConcertCVarDetails> SettingPin = ClientSetting.Pin();
+		if(SettingPin.IsValid())
+		{
+			if (SettingPin->bCVarSyncEnabled)
+			{
+				return LOCTEXT("Transact_ToolTip", "Sets whether this client should receive CVar changes.");
+			}
+		}
+		return GetToolTipNoSync();
+	}
+
+	ECheckBoxState GetTransactState() const
+	{
+		TSharedPtr<FConcertCVarDetails> SettingPin = ClientSetting.Pin();
+		if(SettingPin.IsValid())
+		{
+			if (SettingPin->Settings.bReceiveCVarChanges)
+			{
+				return ECheckBoxState::Checked;
+			}
+		}
+		return ECheckBoxState::Unchecked;
+	}
+
+	void SetTransactState(ECheckBoxState InValue)
+	{
+		TSharedPtr<FConcertCVarDetails> SettingPin = ClientSetting.Pin();
+		if(SettingPin.IsValid())
+		{
+			SettingPin->Settings.bReceiveCVarChanges = ECheckBoxState::Checked == InValue?true:false;
+			NotifyChange(SettingPin);
+		}
+	}
+
+	bool CheckBoxEnabled() const
+	{
+		TSharedPtr<FConcertCVarDetails> Setting = ClientSetting.Pin();
+		if (Setting.IsValid())
+		{
+			return Setting->bCVarSyncEnabled;
+		}
+		return false;
+	}
+
+	void UpdateUI()
+	{
+		TSharedPtr<FConcertCVarDetails> Setting = ClientSetting.Pin();
+		if (Setting.IsValid())
+		{
+			if (Transact.IsValid())
+			{
+				Transact->SetEnabled(Setting->bCVarSyncEnabled);
+				Transact->SetIsChecked(Setting->Settings.bReceiveCVarChanges);
+				Transact->SetToolTipText(GetTransactToolTip());
+			}
+		}
+	}
+
+	FText GetRowToolTip() const
+	{
+		// This is a tooltip for the entire row. Like display name, the tooltip will not update in real time if the user change its
+		// settings. See GetDisplayName() for more info.
+		TSharedPtr<FConcertCVarDetails> SettingPin = ClientSetting.Pin();
+		return SettingPin.IsValid() ? SettingPin->Details.ToDisplayString() : FText();
+	}
+
+	FText GetDisplayName() const
+	{
+		TSharedPtr<FConcertCVarDetails> SettingPin = ClientSetting.Pin();
+		if (SettingPin.IsValid())
+		{
+			// NOTE: The display name doesn't update in real time at the moment because the concert setting are not propagated
+			//       until the client disconnect/reconnect. Since those settings should not change often, this should not
+			//       be a major deal breaker for the users.
+			TSharedPtr<IConcertClientSession> ClientSessionPin = ClientSession.Pin();
+			if (ClientSessionPin.IsValid() && SettingPin->Details.ClientEndpointId == ClientSessionPin->GetSessionClientEndpointId())
+			{
+				return FText::Format(LOCTEXT("ClientDisplayNameIsYouFmt", "{0} (You)"),
+									 FText::FromString(ClientSessionPin->GetLocalClientInfo().DisplayName));
+			}
+
+			// Return the ClientInfo cached.
+			return FText::FromString(SettingPin->Details.ClientInfo.DisplayName);
+		}
+
+		return FText();
+	}
+
+	FSlateFontInfo GetAvatarFont() const
+	{
+		static const FName ButtonIconSyle = TEXT("FontAwesome.10");
+		// This font is used to render a small square box filled with the avatar color.
+		FSlateFontInfo ClientIconFontInfo = FEditorStyle::Get().GetFontStyle(ButtonIconSyle);
+		ClientIconFontInfo.Size = 8;
+		ClientIconFontInfo.OutlineSettings.OutlineSize = 1;
+
+		TSharedPtr<FConcertCVarDetails> SettingPin = ClientSetting.Pin();
+		if (SettingPin.IsValid())
+		{
+			FConcertSessionClientInfo& Client = SettingPin->Details;
+			FLinearColor ClientOutlineColor = Client.ClientInfo.AvatarColor * 0.6f; // Make the font outline darker.
+			ClientOutlineColor.A = Client.ClientInfo.AvatarColor.A; // Put back the original alpha.
+			ClientIconFontInfo.OutlineSettings.OutlineColor = ClientOutlineColor;
+		}
+		else
+		{
+			ClientIconFontInfo.OutlineSettings.OutlineColor = FLinearColor(0.75, 0.75, 0.75); // This is an arbitrary color.
+		}
+
+		return ClientIconFontInfo;
+	}
+
+	FSlateColor GetAvatarColor() const
+	{
+		TSharedPtr<FConcertCVarDetails> ClientSettingPin = ClientSetting.Pin();
+		if (ClientSettingPin.IsValid())
+		{
+			return ClientSettingPin->Details.ClientInfo.AvatarColor;
+		}
+
+		return FSlateColor(FLinearColor(0.75, 0.75, 0.75)); // This is an arbitrary color.
+	}
+
+	void NotifyChange(TSharedPtr<FConcertCVarDetails>& SettingPin)
+	{
+		check(SettingPin.IsValid());
+		const FConcertCVarDetails& Ref = *SettingPin.Get();
+		SettingChangeDelegate.Broadcast(Ref);
+	}
+
+	FOnConcertCVarDetailChanged& SettingChanged()
+	{
+		return SettingChangeDelegate;
+	}
+
+private:
+	TSharedPtr<SCheckBox>		 RecordOnClient;
+	TSharedPtr<SCheckBox>		 Transact;
+
+	TWeakPtr<IConcertSyncClient> SyncClient;
+	TWeakPtr<IConcertClientSession> ClientSession;
+	TWeakPtr<FConcertCVarDetails> ClientSetting;
+	FOnConcertCVarDetailChanged SettingChangeDelegate;
+};
+
+
+FConcertCVarDetails  GetLocalSessionConcertClientAsSetting(
+	UConcertCVarConfig const* Config,
+	TSharedPtr<IConcertClientSession>& Session)
+{
+	UConcertCVarSynchronization const* Sync = GetDefault<UConcertCVarSynchronization>();
+	const FConcertClientInfo &ClientInfo = Session->GetLocalClientInfo();
+	FConcertCVarDetails LocalClientSetting;
+	LocalClientSetting.Details.ClientEndpointId =
+		Session->GetSessionClientEndpointId();
+	LocalClientSetting.Details.ClientInfo = ClientInfo;
+	LocalClientSetting.Settings = Config->LocalSettings;
+	LocalClientSetting.bCVarSyncEnabled = Sync->bSyncCVarTransactions;
+	return MoveTemp(LocalClientSetting);
+}
+
+void FConcertConsoleVariableSessionCustomization::PopulateClientList()
+{
+	Clients.Reset();
+	TSharedPtr<IConcertSyncClient> ConcertSyncClient = IConcertSyncClientModule::Get().GetClient(TEXT("MultiUser"));
+	IConcertClientRef ConcertClient = ConcertSyncClient->GetConcertClient();
+	TSharedPtr<IConcertClientSession> Session = ConcertClient->GetCurrentSession();
+
+	if (Session)
+	{
+		UConcertCVarConfig const* Settings = GetDefault<UConcertCVarConfig>();
+
+		Clients.Reserve(Settings->RemoteDetails.Num()+1);
+
+		FConcertCVarDetails Local = GetLocalSessionConcertClientAsSetting(Settings,Session);
+
+		Clients.Emplace(MakeShared<FConcertCVarDetails>(MoveTemp(Local)));
+		for (FConcertCVarDetails const& Client : Settings->RemoteDetails)
+		{
+			Clients.Emplace(MakeShared<FConcertCVarDetails>(Client));
+		}
+	}
+
+	if (ClientsListViewWeak.IsValid())
+	{
+		TSharedPtr<SListViewCVarDetail> ClientsListView = ClientsListViewWeak.Pin();
+		ClientsListView->RequestListRefresh();
+	}
+}
+
+void FConcertConsoleVariableSessionCustomization::UpdateClientSettings(
+	EConcertClientStatus Status, const FConcertCVarDetails& InSetting)
+{
+	int32 Index = Clients.IndexOfByPredicate([&InSetting](const TSharedPtr<FConcertCVarDetails>& Setting) {
+		return Setting->Details.ClientEndpointId == InSetting.Details.ClientEndpointId;
+	});
+
+	switch(Status)
+	{
+	case EConcertClientStatus::Connected:
+		Clients.Emplace(MakeShared<FConcertCVarDetails>(InSetting));
+		break;
+	case EConcertClientStatus::Updated:
+		check(Index != INDEX_NONE);
+		Clients[Index]->Settings = InSetting.Settings;
+		Clients[Index]->bCVarSyncEnabled = InSetting.bCVarSyncEnabled;
+		break;
+	case EConcertClientStatus::Disconnected:
+		check(Index != INDEX_NONE);
+		Clients.RemoveAt(Index);
+		break;
+	};
+
+	if (ClientsListViewWeak.IsValid())
+	{
+		TSharedPtr<SListViewCVarDetail> ClientsListView = ClientsListViewWeak.Pin();
+		if (Status == EConcertClientStatus::Disconnected
+			|| Status == EConcertClientStatus::Connected )
+		{
+			ClientsListView->RequestListRefresh();
+		}
+		else
+		{
+			TSharedPtr<SConsoleVariableSessionRow> Row = StaticCastSharedPtr<SConsoleVariableSessionRow>(ClientsListView->WidgetFromItem(Clients[Index]));
+			if (Row.IsValid())
+			{
+				Row->UpdateUI();
+			}
+		}
+	}
+}
+
+void FConcertConsoleVariableSessionCustomization::SettingChange(const FConcertCVarDetails &Setting)
+{
+	DetailDelegate.Broadcast(Setting);
+}
+
+void FConcertConsoleVariableSessionCustomization::CustomizeDetails(IDetailLayoutBuilder& DetailLayout)
+{
+	IDetailCategoryBuilder& Settings = DetailLayout.EditCategory(TEXT("Multi-user Console Variable Settings"));
+
+	TSharedPtr<IPropertyHandle> SettingsProperty = DetailLayout.GetProperty(
+		GET_MEMBER_NAME_CHECKED(UConcertCVarConfig, LocalSettings));
+	SettingsProperty->MarkHiddenByCustomization();
+
+
+	TSharedPtr<IPropertyHandle> RemoteSettingsProperty = DetailLayout.GetProperty(
+		GET_MEMBER_NAME_CHECKED(UConcertCVarConfig, RemoteDetails));
+	RemoteSettingsProperty->MarkHiddenByCustomization();
+
+	PopulateClientList();
+	{
+		FDetailWidgetRow& Row = Settings.AddCustomRow(LOCTEXT("MultiUserSettings", "Multi-user Client Settings"));
+		auto HandleGenerateRow = [this](TSharedPtr<FConcertCVarDetails> InClientInfo,
+									const TSharedRef<STableViewBase>& OwnerTable) -> TSharedRef<ITableRow>
+		{
+			TSharedPtr<IConcertSyncClient> ConcertSyncClient = IConcertSyncClientModule::Get().GetClient(TEXT("MultiUser"));
+			IConcertClientRef ConcertClient = ConcertSyncClient->GetConcertClient();
+			TSharedPtr<IConcertClientSession> ConcertClientSession = ConcertClient->GetCurrentSession();
+			TSharedRef<SConsoleVariableSessionRow> Row = SNew(SConsoleVariableSessionRow,
+															  ConcertSyncClient,
+															  ConcertClientSession,
+															  InClientInfo, OwnerTable);
+			Row->SettingChanged().AddSP(this,&FConcertConsoleVariableSessionCustomization::SettingChange);
+			return Row;
+		};
+		Row.WholeRowContent()
+		[
+			SAssignNew(ClientsListViewWeak,SListView<TSharedPtr<FConcertCVarDetails>>)
+			.ItemHeight(20.0f)
+			.SelectionMode(ESelectionMode::None)
+			.ListItemsSource(&Clients)
+			.OnGenerateRow_Lambda(HandleGenerateRow)
+			.HeaderRow
+			(
+				SNew(SHeaderRow)
+				+SHeaderRow::Column(ConsoleVariableSessionUI::DisplayNameColumnName)
+				.DefaultLabel(LOCTEXT("UserDisplayName", "Display Name"))
+				+SHeaderRow::Column(ConsoleVariableSessionUI::TransactColumnName)
+				.DefaultLabel(LOCTEXT("Transact", "Transact Console Variables"))
+			)
+		];
+	}
+};
+
+#undef LOCTEXT_NAMESPACE
diff --git a/Engine/Plugins/Editor/ConsoleVariablesEditor/Source/ConsoleVariablesEditor/Private/MultiUser/ConcertConsoleVariableSessionCustomization.h b/Engine/Plugins/Editor/ConsoleVariablesEditor/Source/ConsoleVariablesEditor/Private/MultiUser/ConcertConsoleVariableSessionCustomization.h
new file mode 100644
--- /dev/null
+++ b/Engine/Plugins/Editor/ConsoleVariablesEditor/Source/ConsoleVariablesEditor/Private/MultiUser/ConcertConsoleVariableSessionCustomization.h
@@ -0,0 +1,40 @@
+// Copyright Epic Games, Inc. All Rights Reserved.
+
+#pragma once
+
+#include "ConcertMessages.h"
+#include "IConcertSession.h"
+#include "IDetailCustomization.h"
+#include "DetailLayoutBuilder.h"
+#include "DetailCategoryBuilder.h"
+#include "ConsoleVariableSyncData.h"
+#include "IPropertyTypeCustomization.h"
+
+#include "Templates/SharedPointer.h"
+#include "Widgets/Views/SListView.h"
+
+DECLARE_MULTICAST_DELEGATE_OneParam(FOnConcertCVarDetailChanged, const FConcertCVarDetails&);
+
+using SListViewCVarDetail = SListView<TSharedPtr<FConcertCVarDetails>>;
+
+class FConcertConsoleVariableSessionCustomization : public IDetailCustomization
+{
+public:
+	void CustomizeDetails(IDetailLayoutBuilder& DetailLayout) override;
+
+	void UpdateClientSettings(EConcertClientStatus Status, const FConcertCVarDetails& RecordSetting);
+	void PopulateClientList();
+
+	FOnConcertCVarDetailChanged& OnSettingChanged()
+	{
+		return DetailDelegate;
+	}
+
+private:
+	void SettingChange(const FConcertCVarDetails &Setting);
+
+	FOnConcertCVarDetailChanged DetailDelegate;
+
+	TArray<TSharedPtr<FConcertCVarDetails>> Clients;
+	TWeakPtr<SListViewCVarDetail>		    ClientsListViewWeak;
+};
diff --git a/Engine/Plugins/Editor/ConsoleVariablesEditor/Source/ConsoleVariablesEditor/Private/MultiUser/ConcertConsoleVariableSyncCustomization.h b/Engine/Plugins/Editor/ConsoleVariablesEditor/Source/ConsoleVariablesEditor/Private/MultiUser/ConcertConsoleVariableSyncCustomization.h
new file mode 100644
--- /dev/null
+++ b/Engine/Plugins/Editor/ConsoleVariablesEditor/Source/ConsoleVariablesEditor/Private/MultiUser/ConcertConsoleVariableSyncCustomization.h
@@ -0,0 +1,91 @@
+// Copyright Epic Games, Inc. All Rights Reserved.
+
+#pragma once
+
+#include "Delegates/Delegate.h"
+#include "IDetailCustomization.h"
+#include "DetailLayoutBuilder.h"
+#include "DetailCategoryBuilder.h"
+#include "MultiUser/ConsoleVariableSyncData.h"
+#include "IPropertyTypeCustomization.h"
+#include "Misc/AssertionMacros.h"
+#include "Types/SlateEnums.h"
+#include "UObject/Script.h"
+#include "UObject/UObjectGlobals.h"
+#include "Widgets/SBoxPanel.h"
+#include "Widgets/SWidget.h"
+#include "Widgets/Layout/SBox.h"
+#include "Widgets/Text/STextBlock.h"
+#include "Widgets/Input/SHyperlink.h"
+#include "DetailWidgetRow.h"
+
+#include "Core/Public/Modules/ModuleManager.h"
+#include "ISettingsModule.h"
+
+#define LOCTEXT_NAMESPACE "ConcertConsoleVariableSyncCustomization"
+
+DECLARE_MULTICAST_DELEGATE_OneParam(FOnCVarSyncPropertyValueChanged, bool);
+
+class FConcertConsoleVariableSyncCustomization : public IDetailCustomization
+{
+public:
+	using ThisType = FConcertConsoleVariableSyncCustomization;
+	static void OpenMultiUserExcludeProperty()
+	{
+		FModuleManager::GetModulePtr<ISettingsModule>("Settings")->ShowViewer("Project", "Plugins", "Concert Sync");
+	}
+
+	static FOnCVarSyncPropertyValueChanged& OnSyncPropertyValueChanged()
+	{
+		static FOnCVarSyncPropertyValueChanged SyncPropertyChanged;
+		return SyncPropertyChanged;
+	}
+
+	void SyncPropertyChanged()
+	{
+		UConcertCVarSynchronization const* Sync = GetDefault<UConcertCVarSynchronization>();
+		OnSyncPropertyValueChanged().Broadcast(Sync->bSyncCVarTransactions);
+	}
+
+	void CustomizeDetails(IDetailLayoutBuilder& DetailLayout) override
+	{
+		IDetailCategoryBuilder& TakeSync = DetailLayout.EditCategory("Multi-user");
+
+		TSharedPtr<IPropertyHandle> SyncTakeRecordingProperty = DetailLayout.GetProperty(
+			GET_MEMBER_NAME_CHECKED(UConcertCVarSynchronization, bSyncCVarTransactions));
+		SyncTakeRecordingProperty->MarkHiddenByCustomization();
+		SyncTakeRecordingProperty->SetOnPropertyValueChanged(FSimpleDelegate::CreateRaw(this, &ThisType::SyncPropertyChanged));
+		{
+			TSharedPtr<SWidget> NameWidget;
+			TSharedPtr<SWidget> ValueWidget;
+
+			IDetailPropertyRow& SyncDisplayProperty = TakeSync.AddProperty(SyncTakeRecordingProperty);
+			SyncDisplayProperty.GetDefaultWidgets(NameWidget,ValueWidget);
+			FDetailWidgetRow& Row = SyncDisplayProperty.CustomWidget(true);
+
+			Row.NameContent()
+			[
+				NameWidget.ToSharedRef()
+			];
+
+			Row.ValueContent()
+			.MinDesiredWidth(250)
+			.MaxDesiredWidth(0)
+			[
+				SNew(SHorizontalBox)
+				+ SHorizontalBox::Slot()
+				.FillWidth(1.0)
+				.HAlign(HAlign_Left)
+				[
+					SNew(SBox)
+					.WidthOverride(250.f)
+					[
+						ValueWidget.ToSharedRef()
+					]
+				]
+			];
+		}
+	}
+};
+
+#undef LOCTEXT_NAMESPACE
diff --git a/Engine/Plugins/Editor/ConsoleVariablesEditor/Source/ConsoleVariablesEditor/Private/MultiUser/ConsoleVariableSync.cpp b/Engine/Plugins/Editor/ConsoleVariablesEditor/Source/ConsoleVariablesEditor/Private/MultiUser/ConsoleVariableSync.cpp
new file mode 100644
--- /dev/null
+++ b/Engine/Plugins/Editor/ConsoleVariablesEditor/Source/ConsoleVariablesEditor/Private/MultiUser/ConsoleVariableSync.cpp
@@ -0,0 +1,467 @@
+// Copyright Epic Games, Inc. All Rights Reserved.
+
+#include "MultiUser/ConsoleVariableSync.h"
+
+#include "ConcertTransportMessages.h"
+#include "Delegates/IDelegateInstance.h"
+#include "IConcertClient.h"
+#include "IConcertSyncClient.h"
+#include "IConcertSession.h"
+#include "IConcertSessionHandler.h"
+#include "IConcertSyncClientModule.h"
+
+#include "MultiUser/ConsoleVariableSyncData.h"
+#include "MultiUser/ConcertConsoleVariableSessionCustomization.h"
+#include "MultiUser/ConcertConsoleVariableSyncCustomization.h"
+
+
+#include "PropertyEditorModule.h"
+#include "Templates/UniquePtr.h"
+
+#include "ConsoleVariablesEditorLog.h"
+
+namespace UE::ConsoleVariables::MultiUser::Private
+{
+
+struct FManagerImpl
+{
+	FManagerImpl()
+	{
+		if (TSharedPtr<IConcertSyncClient> ConcertSyncClient = IConcertSyncClientModule::Get().GetClient(TEXT("MultiUser")))
+		{
+			IConcertClientRef ConcertClient = ConcertSyncClient->GetConcertClient();
+			ConcertClient->OnSessionStartup().AddRaw(this, &FManagerImpl::Register);
+			ConcertClient->OnSessionShutdown().AddRaw(this, &FManagerImpl::Unregister);
+
+			if (TSharedPtr<IConcertClientSession> ConcertClientSession = ConcertClient->GetCurrentSession())
+			{
+				Register(ConcertClientSession.ToSharedRef());
+			}
+
+			ConcertClient->OnSessionConnectionChanged().AddRaw(this, &FManagerImpl::OnSessionConnectionChanged);
+
+			RegisterExtensions();
+
+			bIsInitialized = true;
+		}
+	}
+
+	~FManagerImpl()
+	{
+		if (IConcertSyncClientModule::IsAvailable())
+		{
+			TSharedPtr<IConcertSyncClient> ConcertSyncClient = IConcertSyncClientModule::Get().GetClient(TEXT("MultiUser"));
+			Unregister();
+
+			IConcertClientRef ConcertClient = ConcertSyncClient->GetConcertClient();
+			ConcertClient->OnSessionStartup().RemoveAll(this);
+			ConcertClient->OnSessionShutdown().RemoveAll(this);
+		}
+
+		UnregisterExtensions();
+	}
+
+	void RegisterExtensions()
+	{
+		if (GIsEditor)
+		{
+			FPropertyEditorModule& PropertyEditorModule = FModuleManager::Get().LoadModuleChecked<FPropertyEditorModule>("PropertyEditor");
+
+			SyncDelegate = FConcertConsoleVariableSyncCustomization::OnSyncPropertyValueChanged().AddRaw(
+				this,&FManagerImpl::OnSyncPropertyChange);
+
+			PropertyEditorModule.RegisterCustomClassLayout(
+				UConcertCVarSynchronization::StaticClass()->GetFName(),
+				FOnGetDetailCustomizationInstance::CreateLambda(&MakeShared<FConcertConsoleVariableSyncCustomization>));
+
+			PropertyEditorModule.RegisterCustomClassLayout(
+				UConcertCVarConfig::StaticClass()->GetFName(),
+				FOnGetDetailCustomizationInstance::CreateLambda([this]() {
+					Customization = MakeShared<FConcertConsoleVariableSessionCustomization>();
+					Customization->OnSettingChanged().AddRaw(this,&FManagerImpl::SettingChange);
+					return Customization->AsShared();
+			}));
+	}
+	}
+
+	void UnregisterExtensions()
+	{
+		FPropertyEditorModule* PropertyEditorModule = FModuleManager::Get().GetModulePtr<FPropertyEditorModule>("PropertyEditor");
+		if (PropertyEditorModule)
+		{
+			if (UObjectInitialized())
+			{
+				PropertyEditorModule->UnregisterCustomClassLayout(UConcertCVarSynchronization::StaticClass()->GetFName());
+				PropertyEditorModule->UnregisterCustomClassLayout(UConcertCVarConfig::StaticClass()->GetFName());
+			}
+			FConcertConsoleVariableSyncCustomization::OnSyncPropertyValueChanged().Remove(SyncDelegate);
+		}
+	}
+
+	void OnSyncPropertyChange(bool Value)
+	{
+		TSharedPtr<IConcertClientSession> Session = WeakSession.Pin();
+		if (Session.IsValid())
+		{
+			FConcertCVarSyncChangeEvent OutEvent;
+			OutEvent.EndpointId = Session->GetSessionClientEndpointId();
+			OutEvent.bSyncCVarChanges = Value;
+			Session->SendCustomEvent(OutEvent, Session->GetSessionClientEndpointIds(), EConcertMessageFlags::ReliableOrdered);
+
+			if (Customization.IsValid())
+			{
+				UConcertCVarConfig const* Config = GetDefault<UConcertCVarConfig>();
+				FConcertCVarDetails Item;
+				Item.Details.ClientEndpointId = Session->GetSessionClientEndpointId();
+				Item.Settings = Config->LocalSettings;
+				Item.bCVarSyncEnabled = Value;
+				Customization->UpdateClientSettings(EConcertClientStatus::Updated, Item);
+			}
+		}
+		UConcertCVarSynchronization* Sync = GetMutableDefault<UConcertCVarSynchronization>();
+		Sync->SaveConfig();
+	}
+
+	void ConnectToSession(IConcertClientSession& InSession)
+	{
+		UE_LOG(LogConsoleVariablesEditor, Display, TEXT("Multi-user Console Variable Editor listening to session: %s"), *InSession.GetSessionInfo().SessionName);
+
+		ClientChangeDelegate = InSession.OnSessionClientChanged().AddRaw(this, &FManagerImpl::OnSessionClientChanged);
+
+		SendInitialState(InSession);
+		UpdateSessionClientList();
+	}
+
+	void SendInitialState(IConcertClientSession& InSession)
+	{
+		/** Send Mult-user Console Variable configuration for this client. */
+		UConcertCVarConfig const* CurrentSettings = GetDefault<UConcertCVarConfig>();
+		FConcertCVarChangeEvent OutEvent;
+		OutEvent.EndpointId = InSession.GetSessionClientEndpointId();
+		OutEvent.Settings = CurrentSettings->LocalSettings;
+		InSession.SendCustomEvent(OutEvent, InSession.GetSessionClientEndpointIds(), EConcertMessageFlags::ReliableOrdered);
+
+		/** Send the synchronization state for this client.*/
+		UConcertCVarSynchronization const* SyncState = GetDefault<UConcertCVarSynchronization>();
+		FConcertCVarSyncChangeEvent SyncEvent;
+		SyncEvent.EndpointId = OutEvent.EndpointId;
+		SyncEvent.bSyncCVarChanges = SyncState->bSyncCVarTransactions;
+		InSession.SendCustomEvent(SyncEvent, InSession.GetSessionClientEndpointIds(), EConcertMessageFlags::ReliableOrdered);
+	}
+
+	void OnSessionClientChanged(IConcertClientSession& InSession, EConcertClientStatus ClientStatus, const FConcertSessionClientInfo& InClientInfo)
+	{
+		check(WeakSession.IsValid());
+
+		UConcertCVarConfig* Settings = GetMutableDefault<UConcertCVarConfig>();
+		TArray<FConcertCVarDetails>& Clients = Settings->RemoteDetails;
+
+		int32 Index = Clients.IndexOfByPredicate([&InClientInfo](const FConcertCVarDetails& Setting) {
+			return Setting.Details.ClientEndpointId == InClientInfo.ClientEndpointId;
+		});
+
+		switch(ClientStatus)
+		{
+		case EConcertClientStatus::Connected:
+		{
+			check(Index == INDEX_NONE);
+			FConcertCVarDetails ConcertCVar;
+			ConcertCVar.Details = InClientInfo;
+			Clients.Emplace(MoveTemp(ConcertCVar));
+
+			SendInitialState(InSession);
+		}
+		break;
+		case EConcertClientStatus::Disconnected:
+			break;
+		case EConcertClientStatus::Updated:
+			break;
+		};
+
+		if (Customization.IsValid())
+		{
+			Customization->UpdateClientSettings(ClientStatus, Index == INDEX_NONE ? Clients.Last() : Clients[Index] );
+		}
+
+		if (ClientStatus == EConcertClientStatus::Disconnected && Index != INDEX_NONE)
+		{
+			Clients.RemoveAt(Index);
+		}
+		Settings->SaveConfig();
+	}
+
+
+	void AddRemoteClient(const FConcertSessionClientInfo& ClientInfo)
+	{
+		UConcertCVarConfig* Config = GetMutableDefault<UConcertCVarConfig>();
+		FConcertCVarDetails Setting;
+		Setting.Details = ClientInfo;
+
+		Config->RemoteDetails.Emplace(MoveTemp(Setting));
+	}
+
+	void UpdateSessionClientList()
+	{
+		TSharedPtr<IConcertClientSession> Session = WeakSession.Pin();
+
+		check(Session.IsValid());
+
+		TArray<FConcertSessionClientInfo> ConnectedClients = Session->GetSessionClients();
+
+		UConcertCVarConfig* Config = GetMutableDefault<UConcertCVarConfig>();
+		Config->RemoteDetails.Reset();
+		Config->RemoteDetails.Reserve(ConnectedClients.Num());
+
+		for (FConcertSessionClientInfo& Client : ConnectedClients)
+		{
+			AddRemoteClient(Client);
+		}
+
+		if (Customization)
+		{
+			Customization->PopulateClientList();
+		}
+		Config->SaveConfig();
+	}
+
+	void DisconnectFromSession()
+	{
+		check(WeakSession.IsValid());
+		{
+			TSharedPtr<IConcertClientSession> Session = WeakSession.Pin();
+			UE_LOG(LogConsoleVariablesEditor, Display, TEXT("Multi-user Console Variable Editor disconnecting from Session: %s"), *Session->GetSessionInfo().SessionName);
+			Session->OnSessionClientChanged().Remove(ClientChangeDelegate);
+		}
+	}
+
+	void OnSessionConnectionChanged(IConcertClientSession& InSession, EConcertConnectionStatus ConnectionStatus)
+	{
+		if (ConnectionStatus == EConcertConnectionStatus::Connected)
+		{
+			DisconnectFromSession();
+			ConnectToSession(InSession);
+		}
+		else if (ConnectionStatus == EConcertConnectionStatus::Disconnecting)
+		{
+			DisconnectFromSession();
+		}
+
+		ConnectionChanged.Broadcast(ConnectionStatus);
+	}
+
+	void HandleCVarChange(const FConcertSessionContext& Context, const FConcertCVarChangeEvent& InEvent)
+	{
+		TSharedPtr<IConcertClientSession> Session = WeakSession.Pin();
+		if (Session.IsValid())
+		{
+			UConcertCVarConfig* Config = GetMutableDefault<UConcertCVarConfig>();
+			if (InEvent.EndpointId == Session->GetSessionClientEndpointId())
+			{
+				Config->LocalSettings = InEvent.Settings;
+				if (Customization.IsValid())
+				{
+					FConcertCVarDetails Item;
+					Item.Details.ClientEndpointId = InEvent.EndpointId;
+					Item.Settings = InEvent.Settings;
+					Item.bCVarSyncEnabled = GetDefault<UConcertCVarSynchronization>()->bSyncCVarTransactions;
+					Customization->UpdateClientSettings(EConcertClientStatus::Updated, Item);
+				}
+
+			}
+			else
+			{
+				for (FConcertCVarDetails& Item : Config->RemoteDetails)
+				{
+					if (Item.Details.ClientEndpointId == InEvent.EndpointId)
+					{
+						Item.Settings = InEvent.Settings;
+						if (Customization.IsValid())
+						{
+							Customization->UpdateClientSettings(EConcertClientStatus::Updated, Item);
+						}
+					}
+				}
+			}
+		}
+	}
+
+	void HandleCVarSync(const FConcertSessionContext& Context, const FConcertCVarSyncChangeEvent& InEvent)
+	{
+		TSharedPtr<IConcertClientSession> Session = WeakSession.Pin();
+		if (Session.IsValid())
+		{
+			UConcertCVarConfig* Config = GetMutableDefault<UConcertCVarConfig>();
+			for(FConcertCVarDetails& Item : Config->RemoteDetails)
+			{
+				if (Item.Details.ClientEndpointId == InEvent.EndpointId
+					&& Item.bCVarSyncEnabled != InEvent.bSyncCVarChanges)
+				{
+					Item.bCVarSyncEnabled = InEvent.bSyncCVarChanges;
+					if (Customization.IsValid())
+					{
+						Customization->UpdateClientSettings(EConcertClientStatus::Updated, Item);
+					}
+				}
+			}
+		}
+	}
+
+	void HandleCVarSet(const FConcertSessionContext& Context, const FConcertSetConsoleVariableEvent& InEvent)
+	{
+		if (CanReceiveConsoleVariable())
+		{
+			RemoteCVarChanged.Broadcast(InEvent.Variable, InEvent.Value);
+		}
+	}
+
+	/** 
+	 * Register a new multi-user session. 
+	 */
+	void Register(TSharedRef<IConcertClientSession> InSession)
+	{
+		WeakSession = InSession;
+
+		InSession->RegisterCustomEventHandler<FConcertCVarChangeEvent>(this, &FManagerImpl::HandleCVarChange);
+		InSession->RegisterCustomEventHandler<FConcertCVarSyncChangeEvent>(this, &FManagerImpl::HandleCVarSync);
+		InSession->RegisterCustomEventHandler<FConcertSetConsoleVariableEvent>(this, &FManagerImpl::HandleCVarSet);
+	}
+
+	/**
+	 * Unregister a multi-user session (if one exists)
+	 */
+	void Unregister()
+	{
+		if (TSharedPtr<IConcertClientSession> Session = WeakSession.Pin())
+		{
+			Session->UnregisterCustomEventHandler<FConcertSetConsoleVariableEvent>(this);
+			Session->UnregisterCustomEventHandler<FConcertCVarSyncChangeEvent>(this);
+			Session->UnregisterCustomEventHandler<FConcertCVarChangeEvent>(this);
+		}
+
+		WeakSession.Reset();
+	}
+
+	/** 
+	 * Unregister call that is invoked by Multi-user when it is shutdown. 
+	 */
+	void Unregister(TSharedRef<IConcertClientSession> InSession)
+	{
+		Unregister();
+	}
+
+	/**
+	 * Returns true if this client can receive console variable changes.
+	 */
+	bool CanReceiveConsoleVariable() const
+	{
+		UConcertCVarSynchronization const* Sync = GetDefault<UConcertCVarSynchronization>();
+		UConcertCVarConfig const* Config = GetDefault<UConcertCVarConfig>();
+		return Sync->bSyncCVarTransactions && Config->LocalSettings.bReceiveCVarChanges;
+	}
+
+	void SettingChange(const FConcertCVarDetails& InChanged)
+	{
+		TSharedPtr<IConcertClientSession> Session = WeakSession.Pin();
+		if (Session.IsValid())
+		{
+			UConcertCVarConfig* Config = GetMutableDefault<UConcertCVarConfig>();
+			if (InChanged.Details.ClientEndpointId == Session->GetSessionClientEndpointId())
+			{
+				Config->LocalSettings = InChanged.Settings;
+			}
+			else
+			{
+				for(FConcertCVarDetails& Item : Config->RemoteDetails)
+				{
+					if (Item.Details.ClientEndpointId == InChanged.Details.ClientEndpointId)
+					{
+						Item.Settings = InChanged.Settings;
+					}
+				}
+			}
+
+			Config->SaveConfig();
+
+			FConcertCVarChangeEvent OutEvent;
+			OutEvent.Settings   = InChanged.Settings;
+			OutEvent.EndpointId = InChanged.Details.ClientEndpointId;
+			Session->SendCustomEvent(OutEvent, Session->GetSessionClientEndpointIds(), EConcertMessageFlags::ReliableOrdered);
+		}
+	}
+
+	/**
+	 * Sends the given console variable name with the specified value to all connected endpoints.  It is up to the
+	 * endpoint to implement the change locally based on configured synchronization state.
+	 */
+	void SendConsoleVariableChange(FString InName, FString InValue)
+	{
+		TSharedPtr<IConcertClientSession> Session = WeakSession.Pin();
+		if (!bIsEnabled || !Session.IsValid())
+		{
+			return;
+		}
+
+		FConcertSetConsoleVariableEvent OutEvent{MoveTemp(InName), MoveTemp(InValue)};
+		Session->SendCustomEvent(OutEvent, Session->GetSessionClientEndpointIds(),
+								 EConcertMessageFlags::ReliableOrdered | EConcertMessageFlags::UniqueId);
+	}
+
+	/** Reference to the customization object so that we can sync with the session.*/
+	TSharedPtr<FConcertConsoleVariableSessionCustomization> Customization;
+
+	/** Weak pointer to the client session with which to send events. May be null or stale. */
+	TWeakPtr<IConcertClientSession> WeakSession;
+
+	/** Delegate when property sync changes. */
+	FDelegateHandle SyncDelegate;
+
+	/** Delegate for any changes in client state. */
+	FDelegateHandle	ClientChangeDelegate;
+
+	/** Multicast delegate invoked when Multi-user console variable change event is received. */
+	FOnRemoteCVarChange RemoteCVarChanged;
+
+	/** Multi-user connection change delegate */
+	FOnMultiUserConnectionChange ConnectionChanged;
+
+	/** Indicates if the multi-user console synchronization is enabled. */
+	bool bIsEnabled = true;
+
+	/** True if the Concert module was found and customizations have been registered. */
+	bool bIsInitialized = false;
+};
+
+
+FManager::FManager()
+{
+	Implementation = MakeUnique<FManagerImpl>();
+}
+
+FManager::~FManager() = default;
+
+FOnRemoteCVarChange& FManager::OnRemoteCVarChange()
+{
+	return Implementation->RemoteCVarChanged;
+}
+
+FOnMultiUserConnectionChange& FManager::OnConnectionChange()
+{
+	return Implementation->ConnectionChanged;
+}
+
+void FManager::SendConsoleVariableChange(FString InName, FString InValue)
+{
+	Implementation->SendConsoleVariableChange(MoveTemp(InName), MoveTemp(InValue));
+}
+
+void FManager::SetEnableMultiUserSupport(bool bIsEnabled)
+{
+	Implementation->bIsEnabled = bIsEnabled;
+}
+
+bool FManager::IsInitialized() const
+{
+	return Implementation->bIsInitialized;
+}
+};
+
diff --git a/Engine/Plugins/Editor/ConsoleVariablesEditor/Source/ConsoleVariablesEditor/Private/MultiUser/ConsoleVariableSync.h b/Engine/Plugins/Editor/ConsoleVariablesEditor/Source/ConsoleVariablesEditor/Private/MultiUser/ConsoleVariableSync.h
new file mode 100644
--- /dev/null
+++ b/Engine/Plugins/Editor/ConsoleVariablesEditor/Source/ConsoleVariablesEditor/Private/MultiUser/ConsoleVariableSync.h
@@ -0,0 +1,49 @@
+// Copyright Epic Games, Inc. All Rights Reserved.
+
+#pragma once
+
+#include "Containers/UnrealString.h"
+#include "Delegates/Delegate.h"
+#include "Delegates/DelegateCombinations.h"
+#include "Templates/UniquePtr.h"
+
+#include "ConcertMessages.h"
+
+DECLARE_MULTICAST_DELEGATE_TwoParams(FOnRemoteCVarChange, FString, FString);
+DECLARE_MULTICAST_DELEGATE_OneParam(FOnMultiUserConnectionChange, EConcertConnectionStatus);
+
+namespace UE::ConsoleVariables::MultiUser::Private
+{
+
+struct FManagerImpl;
+
+struct FManager
+{
+	FManager();
+	~FManager();
+
+	// Do not allow this to be copied or moved.
+	FManager(const FManager&) = delete;
+	FManager(FManager&&) = delete;
+	FManager &operator=(const FManager&) = delete;
+	FManager &operator=(FManager&&) = delete;
+
+	/** Delegate that is invoked when a remote client has sent a new console variable value.*/
+	FOnRemoteCVarChange& OnRemoteCVarChange();
+
+	/** Delegate that is invoked when the connection status changes for this client. */
+	FOnMultiUserConnectionChange& OnConnectionChange();
+
+	/** Sends the named console variable with value to all connected Multi-user clients */
+	void SendConsoleVariableChange(FString InName, FString InValue);
+
+	/** Enables / disables multi-user message handling. */
+	void SetEnableMultiUserSupport(bool bIsEnabled);
+
+	bool IsInitialized() const;
+	
+private:
+	TUniquePtr<FManagerImpl> Implementation;
+};
+
+}
diff --git a/Engine/Plugins/Editor/ConsoleVariablesEditor/Source/ConsoleVariablesEditor/Private/MultiUser/ConsoleVariableSyncData.h b/Engine/Plugins/Editor/ConsoleVariablesEditor/Source/ConsoleVariablesEditor/Private/MultiUser/ConsoleVariableSyncData.h
new file mode 100644
--- /dev/null
+++ b/Engine/Plugins/Editor/ConsoleVariablesEditor/Source/ConsoleVariablesEditor/Private/MultiUser/ConsoleVariableSyncData.h
@@ -0,0 +1,99 @@
+// Copyright Epic Games, Inc. All Rights Reserved.
+
+#pragma once
+
+#include "CoreMinimal.h"
+#include "Misc/Guid.h"
+
+#include "ConcertMessageData.h"
+#include "ConsoleVariableSyncData.generated.h"
+
+UCLASS(config=Engine, DisplayName="Multi-user Console Variable Synchronization")
+class UConcertCVarSynchronization : public UObject
+{
+public:
+	GENERATED_BODY()
+
+	UPROPERTY(config,EditAnywhere,BlueprintReadWrite,Category="Multi-user",DisplayName="Multi-user Console Variables Synchronization")
+	bool bSyncCVarTransactions = true;
+};
+
+USTRUCT()
+struct FConcertCVarSettings
+{
+	GENERATED_BODY();
+
+	/** Indicates if the given node can receive console variable values from other clients. */
+	UPROPERTY(config,EditAnywhere,Category="Multi-user Console Variable Settings")
+	bool bReceiveCVarChanges = true;
+};
+
+USTRUCT()
+struct FConcertCVarDetails
+{
+	GENERATED_BODY()
+
+	UPROPERTY()
+	FConcertSessionClientInfo Details;
+
+	UPROPERTY()
+	bool bCVarSyncEnabled = true;
+
+	UPROPERTY()
+	FConcertCVarSettings Settings;
+};
+
+UCLASS(config=Engine, DisplayName="Multi-user Console Variable Settings")
+class UConcertCVarConfig : public UObject
+{
+	GENERATED_BODY()
+
+public:
+	UPROPERTY(config,EditAnywhere,Category="Multi-user Console Variable Settings")
+	FConcertCVarSettings LocalSettings;
+
+	UPROPERTY(Transient,EditAnywhere,Category="Multi-user Console Variable Settings")
+	TArray<FConcertCVarDetails> RemoteDetails;
+};
+
+USTRUCT()
+struct FConcertCVarChangeEvent
+{
+	GENERATED_BODY()
+
+	/** The endpoint id sending this event. */
+	UPROPERTY()
+	FGuid EndpointId;
+
+	/** The Console Variable multi-user settings for the given EndpointId. */
+	UPROPERTY()
+	FConcertCVarSettings Settings;
+};
+
+USTRUCT()
+struct FConcertCVarSyncChangeEvent
+{
+	GENERATED_BODY()
+
+	/** The endpoint id sending this event. */
+	UPROPERTY()
+	FGuid EndpointId;
+
+	/** The value of the console variable mult-user sync for the above EndpointId. */
+	UPROPERTY()
+	bool bSyncCVarChanges = true;
+};
+
+USTRUCT()
+struct FConcertSetConsoleVariableEvent
+{
+	GENERATED_BODY()
+
+	/** Console variable name to apply value change. */
+	UPROPERTY()
+	FString Variable;
+
+	/** The value to set the named console variable. */
+	UPROPERTY()
+	FString Value;
+};
diff --git a/Engine/Plugins/Editor/ConsoleVariablesEditor/Source/ConsoleVariablesEditor/Private/Views/List/ConsoleVariablesEditorList.cpp b/Engine/Plugins/Editor/ConsoleVariablesEditor/Source/ConsoleVariablesEditor/Private/Views/List/ConsoleVariablesEditorList.cpp
new file mode 100644
--- /dev/null
+++ b/Engine/Plugins/Editor/ConsoleVariablesEditor/Source/ConsoleVariablesEditor/Private/Views/List/ConsoleVariablesEditorList.cpp
@@ -0,0 +1,49 @@
+// Copyright Epic Games, Inc. All Rights Reserved.
+
+#include "Views/List/ConsoleVariablesEditorList.h"
+
+#include "Views/List/SConsoleVariablesEditorList.h"
+
+FConsoleVariablesEditorList::~FConsoleVariablesEditorList()
+{
+	ListWidget.Reset();
+}
+
+TSharedRef<SWidget> FConsoleVariablesEditorList::GetOrCreateWidget()
+{
+	if (!ListWidget.IsValid())
+	{
+		SAssignNew(ListWidget, SConsoleVariablesEditorList, SharedThis(this));
+	}
+
+	return ListWidget.ToSharedRef();
+}
+
+void FConsoleVariablesEditorList::SetSearchString(const FString& SearchString)
+{
+	ListWidget->SetSearchStringInSearchInputField(SearchString);
+}
+
+void FConsoleVariablesEditorList::RebuildList(const FString& InConsoleCommandToScrollTo, bool bShouldCacheValues) const
+{
+	if (ListWidget.IsValid())
+	{
+		ListWidget->RebuildList(InConsoleCommandToScrollTo, bShouldCacheValues);
+	}
+}
+
+void FConsoleVariablesEditorList::RefreshList() const
+{
+	if (ListWidget.IsValid())
+	{
+		ListWidget->RefreshList();
+	}
+}
+
+void FConsoleVariablesEditorList::UpdatePresetValuesForSave(TObjectPtr<UConsoleVariablesAsset> InAsset) const
+{
+	if (ListWidget.IsValid())
+	{
+		ListWidget->UpdatePresetValuesForSave(InAsset);
+	}
+}
diff --git a/Engine/Plugins/Editor/ConsoleVariablesEditor/Source/ConsoleVariablesEditor/Private/Views/List/ConsoleVariablesEditorList.h b/Engine/Plugins/Editor/ConsoleVariablesEditor/Source/ConsoleVariablesEditor/Private/Views/List/ConsoleVariablesEditorList.h
new file mode 100644
--- /dev/null
+++ b/Engine/Plugins/Editor/ConsoleVariablesEditor/Source/ConsoleVariablesEditor/Private/Views/List/ConsoleVariablesEditorList.h
@@ -0,0 +1,62 @@
+// Copyright Epic Games, Inc. All Rights Reserved.
+
+#pragma once
+
+#include "Widgets/SWidget.h"
+
+struct FConsoleVariablesEditorListRow;
+typedef TSharedPtr<FConsoleVariablesEditorListRow> FConsoleVariablesEditorListRowPtr;
+
+class SConsoleVariablesEditorList;
+class UConsoleVariablesAsset;
+
+class FConsoleVariablesEditorList : public TSharedFromThis<FConsoleVariablesEditorList>
+{
+public:
+
+	enum class EConsoleVariablesEditorListMode : uint8
+	{
+		// We're displaying the cvars listed in the loaded (or default) preset
+		Preset = 1,
+		// We're displaying the cvars that match the criteria of the global search
+		GlobalSearch = 2
+	};
+
+	FConsoleVariablesEditorList() {}
+
+	~FConsoleVariablesEditorList();
+
+	TSharedRef<SWidget> GetOrCreateWidget();
+
+	void SetSearchString(const FString& SearchString);
+
+	EConsoleVariablesEditorListMode GetListMode() const
+	{
+		return CurrentListMode;
+	}
+
+	void SetListMode(EConsoleVariablesEditorListMode NewListMode)
+	{
+		CurrentListMode = NewListMode;;
+	}
+
+	/*
+	 * Regenerate the list items and refresh the list. Call when adding or removing variables.
+	 * @param bShouldCacheValues If true, the current list's current values will be cached and then restored when the list is rebuilt. Otherwise preset values will be used.
+	 */
+	void RebuildList(const FString& InConsoleCommandToScrollTo = "", bool bShouldCacheValues = true) const;
+
+	/**
+	 * Refresh filters and sorting.
+	 * Useful for when the list state has gone stale but the variable count has not changed.
+	 */
+	void RefreshList() const;
+
+	/** Updates the saved values in a UConsoleVariablesAsset so that the command/value map can be saved to disk */
+	void UpdatePresetValuesForSave(TObjectPtr<UConsoleVariablesAsset> InAsset) const;
+
+private:
+
+	EConsoleVariablesEditorListMode CurrentListMode = EConsoleVariablesEditorListMode::Preset;
+	TSharedPtr<SConsoleVariablesEditorList> ListWidget;
+};
diff --git a/Engine/Plugins/Editor/ConsoleVariablesEditor/Source/ConsoleVariablesEditor/Private/Views/List/ConsoleVariablesEditorListFilters/ConsoleVariablesEditorListFilter_ModifiedVariables.h b/Engine/Plugins/Editor/ConsoleVariablesEditor/Source/ConsoleVariablesEditor/Private/Views/List/ConsoleVariablesEditorListFilters/ConsoleVariablesEditorListFilter_ModifiedVariables.h
new file mode 100644
--- /dev/null
+++ b/Engine/Plugins/Editor/ConsoleVariablesEditor/Source/ConsoleVariablesEditor/Private/Views/List/ConsoleVariablesEditorListFilters/ConsoleVariablesEditorListFilter_ModifiedVariables.h
@@ -0,0 +1,46 @@
+// Copyright Epic Games, Inc. All Rights Reserved.
+
+#pragma once
+
+#include "IConsoleVariablesEditorListFilter.h"
+
+#define LOCTEXT_NAMESPACE "ConsoleVariablesEditor"
+
+class ConsoleVariablesEditorListFilter_ModifiedVariables : public IConsoleVariablesEditorListFilter
+{
+public:
+
+	ConsoleVariablesEditorListFilter_ModifiedVariables()
+	{
+		// This filter should be off by default
+		SetFilterActive(false);
+		SetFilterMatchType(EConsoleVariablesEditorListFilterMatchType::MatchAll);
+	}
+
+	virtual FString GetFilterName() override
+	{
+		return "Show Only Modified Variables";
+	}
+
+	virtual FText GetFilterButtonLabel() override
+	{
+		return LOCTEXT("ShowOnlyModifiedVariablesFilter", "Show Only Modified");
+	}
+
+	virtual FText GetFilterButtonToolTip() override
+	{
+		return LOCTEXT("ShowSourceTextFilterTooltipFormat", "Show only rows that have a current value that differs from the preset value.");
+	}
+
+	virtual bool DoesItemPassFilter(const FConsoleVariablesEditorListRowPtr& InItem) override
+	{
+		if (InItem.IsValid())
+		{
+			return InItem->DoesCurrentValueDifferFromPresetValue();
+		}
+
+		return false;
+	}
+};
+
+#undef LOCTEXT_NAMESPACE
diff --git a/Engine/Plugins/Editor/ConsoleVariablesEditor/Source/ConsoleVariablesEditor/Private/Views/List/ConsoleVariablesEditorListFilters/ConsoleVariablesEditorListFilter_SourceText.h b/Engine/Plugins/Editor/ConsoleVariablesEditor/Source/ConsoleVariablesEditor/Private/Views/List/ConsoleVariablesEditorListFilters/ConsoleVariablesEditorListFilter_SourceText.h
new file mode 100644
--- /dev/null
+++ b/Engine/Plugins/Editor/ConsoleVariablesEditor/Source/ConsoleVariablesEditor/Private/Views/List/ConsoleVariablesEditorListFilters/ConsoleVariablesEditorListFilter_SourceText.h
@@ -0,0 +1,50 @@
+// Copyright Epic Games, Inc. All Rights Reserved.
+
+#pragma once
+
+#include "IConsoleVariablesEditorListFilter.h"
+
+#define LOCTEXT_NAMESPACE "ConsoleVariablesEditor"
+
+class FConsoleVariablesEditorListFilter_SourceText : public IConsoleVariablesEditorListFilter
+{
+public:
+
+	FConsoleVariablesEditorListFilter_SourceText(const FString& InFilterString)
+	: FilterString(InFilterString){}
+
+	virtual FString GetFilterName() override
+	{
+		return FilterString;
+	}
+
+	virtual FText GetFilterButtonLabel() override
+	{
+		return FText::Format(LOCTEXT("ShowSourceTextFilterFormat", "Show {0}"), FText::FromString(GetFilterName()));
+	}
+
+	virtual FText GetFilterButtonToolTip() override
+	{
+		return FText::Format(
+			LOCTEXT("ShowSourceTextFilterTooltipFormat", "Show rows that have a Source field matching '{0}'"),
+			FText::FromString(GetFilterName()));
+	}
+
+	virtual bool DoesItemPassFilter(const FConsoleVariablesEditorListRowPtr& InItem) override
+	{
+		if (InItem.IsValid(); const TSharedPtr<FConsoleVariablesEditorCommandInfo> PinnedCommand = InItem->GetCommandInfo().Pin())
+		{
+			const FString& StringToSearch = PinnedCommand->GetSourceAsText().ToString();
+
+			return StringToSearch.Contains(FilterString);
+		}
+
+		return false;
+	}
+	
+private:
+
+	FString FilterString;
+};
+
+#undef LOCTEXT_NAMESPACE
diff --git a/Engine/Plugins/Editor/ConsoleVariablesEditor/Source/ConsoleVariablesEditor/Private/Views/List/ConsoleVariablesEditorListFilters/IConsoleVariablesEditorListFilter.h b/Engine/Plugins/Editor/ConsoleVariablesEditor/Source/ConsoleVariablesEditor/Private/Views/List/ConsoleVariablesEditorListFilters/IConsoleVariablesEditorListFilter.h
new file mode 100644
--- /dev/null
+++ b/Engine/Plugins/Editor/ConsoleVariablesEditor/Source/ConsoleVariablesEditor/Private/Views/List/ConsoleVariablesEditorListFilters/IConsoleVariablesEditorListFilter.h
@@ -0,0 +1,63 @@
+// Copyright Epic Games, Inc. All Rights Reserved.
+
+#pragma once
+
+#include "../ConsoleVariablesEditorListRow.h"
+
+class IConsoleVariablesEditorListFilter
+{
+public:
+
+	enum class EConsoleVariablesEditorListFilterMatchType
+	{
+		// Filters of this type are placed into a group in which if any of the group's filters pass, the row is considered passing
+		MatchAny,
+		// Filters of this type are placed into a group in which all of the group's filters must pass for the row to be considered passing
+		MatchAll
+	};
+	
+	virtual ~IConsoleVariablesEditorListFilter() = default;
+
+	virtual FString GetFilterName() = 0;
+
+	/** Returns localized text to display for when this filter' active state is defined by a toggle button. */
+	virtual FText GetFilterButtonLabel() = 0;
+	/** Returns localized text to display when the user mouses over a toggle button that defines whether this filter is active. */
+	virtual FText GetFilterButtonToolTip() = 0;
+	
+	void SetFilterActive(const bool bNewEnabled)
+	{
+		bIsFilterActive = bNewEnabled;
+	}
+
+	void ToggleFilterActive()
+	{
+		bIsFilterActive = !bIsFilterActive;
+	}
+
+	bool GetIsFilterActive() const
+	{
+		return bIsFilterActive;
+	}
+
+	void SetFilterMatchType(const EConsoleVariablesEditorListFilterMatchType InType)
+	{
+		FilterMatchType = InType;
+	}
+
+	EConsoleVariablesEditorListFilterMatchType GetFilterMatchType() const
+	{
+		return FilterMatchType;
+	}
+
+	virtual bool DoesItemPassFilter(const FConsoleVariablesEditorListRowPtr& InItem)
+	{
+		return false;
+	}
+
+private:
+
+	bool bIsFilterActive = true;
+	EConsoleVariablesEditorListFilterMatchType FilterMatchType = EConsoleVariablesEditorListFilterMatchType::MatchAny;
+			
+};
diff --git a/Engine/Plugins/Editor/ConsoleVariablesEditor/Source/ConsoleVariablesEditor/Private/Views/List/ConsoleVariablesEditorListRow.cpp b/Engine/Plugins/Editor/ConsoleVariablesEditor/Source/ConsoleVariablesEditor/Private/Views/List/ConsoleVariablesEditorListRow.cpp
new file mode 100644
--- /dev/null
+++ b/Engine/Plugins/Editor/ConsoleVariablesEditor/Source/ConsoleVariablesEditor/Private/Views/List/ConsoleVariablesEditorListRow.cpp
@@ -0,0 +1,357 @@
+// Copyright Epic Games, Inc. All Rights Reserved.
+
+#include "ConsoleVariablesEditorListRow.h"
+
+#include "ConsoleVariablesAsset.h"
+#include "ConsoleVariablesEditorModule.h"
+
+#include "Algo/AllOf.h"
+#include "Algo/AnyOf.h"
+#include "Views/List/ConsoleVariablesEditorList.h"
+#include "Views/List/SConsoleVariablesEditorList.h"
+
+FConsoleVariablesEditorListRow::~FConsoleVariablesEditorListRow()
+{
+	FlushReferences();
+}
+
+void FConsoleVariablesEditorListRow::FlushReferences()
+{
+	if (ChildRows.Num())
+	{
+		ChildRows.Empty();
+	}
+}
+
+TWeakPtr<FConsoleVariablesEditorCommandInfo> FConsoleVariablesEditorListRow::GetCommandInfo() const
+{
+	return CommandInfo;
+}
+
+FConsoleVariablesEditorListRow::EConsoleVariablesEditorListRowType FConsoleVariablesEditorListRow::GetRowType() const
+{
+	return RowType;
+}
+
+int32 FConsoleVariablesEditorListRow::GetChildDepth() const
+{
+	return ChildDepth;
+}
+
+void FConsoleVariablesEditorListRow::SetChildDepth(const int32 InDepth)
+{
+	ChildDepth = InDepth;
+}
+
+int32 FConsoleVariablesEditorListRow::GetSortOrder() const
+{
+	return SortOrder;
+}
+
+void FConsoleVariablesEditorListRow::SetSortOrder(const int32 InNewOrder)
+{
+	SortOrder = InNewOrder;
+}
+
+TWeakPtr<FConsoleVariablesEditorListRow> FConsoleVariablesEditorListRow::GetDirectParentRow() const
+{
+	return DirectParentRow;
+}
+
+void FConsoleVariablesEditorListRow::SetDirectParentRow(
+	const TWeakPtr<FConsoleVariablesEditorListRow>& InDirectParentRow)
+{
+	DirectParentRow = InDirectParentRow;
+}
+
+const TArray<FConsoleVariablesEditorListRowPtr>& FConsoleVariablesEditorListRow::GetChildRows() const
+{
+	return ChildRows;
+}
+
+int32 FConsoleVariablesEditorListRow::GetChildCount() const
+{
+	return ChildRows.Num();
+}
+
+void FConsoleVariablesEditorListRow::SetChildRows(const TArray<FConsoleVariablesEditorListRowPtr>& InChildRows)
+{
+	ChildRows = InChildRows;
+}
+
+void FConsoleVariablesEditorListRow::AddToChildRows(const FConsoleVariablesEditorListRowPtr& InRow)
+{
+	ChildRows.Add(InRow);
+}
+
+void FConsoleVariablesEditorListRow::InsertChildRowAtIndex(const FConsoleVariablesEditorListRowPtr& InRow,
+                                                           const int32 AtIndex)
+{
+	ChildRows.Insert(InRow, AtIndex);
+}
+
+bool FConsoleVariablesEditorListRow::GetIsTreeViewItemExpanded() const
+{
+	return bIsTreeViewItemExpanded;
+}
+
+void FConsoleVariablesEditorListRow::SetIsTreeViewItemExpanded(const bool bNewExpanded)
+{
+	bIsTreeViewItemExpanded = bNewExpanded;
+}
+
+bool FConsoleVariablesEditorListRow::GetShouldFlashOnScrollIntoView() const
+{
+	return bShouldFlashOnScrollIntoView;
+}
+
+void FConsoleVariablesEditorListRow::SetShouldFlashOnScrollIntoView(const bool bNewShouldFlashOnScrollIntoView)
+{
+	bShouldFlashOnScrollIntoView = bNewShouldFlashOnScrollIntoView;
+}
+
+bool FConsoleVariablesEditorListRow::GetShouldExpandAllChildren() const
+{
+	return bShouldExpandAllChildren;
+}
+
+void FConsoleVariablesEditorListRow::SetShouldExpandAllChildren(const bool bNewShouldExpandAllChildren)
+{
+	bShouldExpandAllChildren = bNewShouldExpandAllChildren;
+}
+
+void FConsoleVariablesEditorListRow::ResetToStartupValueAndSource() const
+{
+	GetCommandInfo().Pin()->ExecuteCommand(GetCommandInfo().Pin()->StartupValueAsString);
+	GetCommandInfo().Pin()->SetSourceFlag(GetCommandInfo().Pin()->StartupSource);
+}
+
+const FString& FConsoleVariablesEditorListRow::GetPresetValue() const
+{
+	return PresetValue;
+}
+
+void FConsoleVariablesEditorListRow::SetPresetValue(const FString& InPresetValue)
+{
+	PresetValue = InPresetValue;
+}
+
+bool FConsoleVariablesEditorListRow::MatchSearchTokensToSearchTerms(
+	const TArray<FString> InTokens, ESearchCase::Type InSearchCase)
+{
+	// If the search is cleared we'll consider the row to pass search
+	bool bMatchFound = InTokens.Num() == 0;
+
+	if (!bMatchFound)
+	{
+		const TSharedPtr<FConsoleVariablesEditorCommandInfo> PinnedInfo = CommandInfo.Pin();
+
+		FString SearchTerms = PinnedInfo->Command + PinnedInfo->GetSourceAsText().ToString();
+
+		if (const IConsoleVariable* AsVariable = PinnedInfo->GetConsoleVariablePtr())
+		{
+			// Add ValueAsString and HelpText to SearchTerms
+			SearchTerms += AsVariable->GetString() + PinnedInfo->GetHelpText();
+		}
+
+		// Match any
+		for (const FString& Token : InTokens)
+		{
+			// Match all of these
+			const FString SpaceDelimiter = " ";
+			TArray<FString> OutSpacedArray;
+			if (Token.Contains(SpaceDelimiter) && Token.ParseIntoArray(OutSpacedArray, *SpaceDelimiter, true) > 1)
+			{
+				bMatchFound = Algo::AllOf(OutSpacedArray, [&SearchTerms, InSearchCase](const FString& Comparator)
+				{
+					return SearchTerms.Contains(Comparator, InSearchCase);
+				});
+			}
+			else
+			{
+				bMatchFound = SearchTerms.Contains(Token, InSearchCase);
+			}
+
+			if (bMatchFound)
+			{
+				break;
+			}
+		}
+	}
+
+	bDoesRowMatchSearchTerms = bMatchFound;
+
+	return bMatchFound;
+}
+
+void FConsoleVariablesEditorListRow::ExecuteSearchOnChildNodes(const FString& SearchString) const
+{
+	TArray<FString> Tokens;
+
+	SearchString.ParseIntoArray(Tokens, TEXT(" "), true);
+
+	ExecuteSearchOnChildNodes(Tokens);
+}
+
+void FConsoleVariablesEditorListRow::ExecuteSearchOnChildNodes(const TArray<FString>& Tokens) const
+{
+	for (const FConsoleVariablesEditorListRowPtr& ChildRow : GetChildRows())
+	{
+		if (!ensure(ChildRow.IsValid()))
+		{
+			continue;
+		}
+
+		if (ChildRow->GetRowType() == EConsoleVariablesEditorListRowType::CommandGroup)
+		{
+			if (ChildRow->MatchSearchTokensToSearchTerms(Tokens))
+			{
+				// If the group name matches then we pass an empty string to search child nodes since we want them all to be visible
+				ChildRow->ExecuteSearchOnChildNodes("");
+			}
+			else
+			{
+				// Otherwise we iterate over all child nodes to determine which should and should not be visible
+				ChildRow->ExecuteSearchOnChildNodes(Tokens);
+			}
+		}
+		else
+		{
+			ChildRow->MatchSearchTokensToSearchTerms(Tokens);
+		}
+	}
+}
+
+bool FConsoleVariablesEditorListRow::GetDoesRowPassFilters() const
+{
+	return bDoesRowPassFilters;
+}
+
+void FConsoleVariablesEditorListRow::SetDoesRowPassFilters(const bool bPass)
+{
+	bDoesRowPassFilters = bPass;
+}
+
+bool FConsoleVariablesEditorListRow::DoesCurrentValueDifferFromPresetValue() const
+{
+	return bDoesCurrentValueDifferFromPresetValue;
+}
+
+void FConsoleVariablesEditorListRow::SetDoesCurrentValueDifferFromPresetValue(
+	const bool bNewValue)
+{
+	bDoesCurrentValueDifferFromPresetValue = bNewValue;
+}
+
+bool FConsoleVariablesEditorListRow::GetIsSelected() const
+{
+	return bIsSelected;
+}
+
+void FConsoleVariablesEditorListRow::SetIsSelected(const bool bNewSelected)
+{
+	bIsSelected = bNewSelected;
+}
+
+ECheckBoxState FConsoleVariablesEditorListRow::GetWidgetCheckedState() const
+{
+	return WidgetCheckedState;
+}
+
+void FConsoleVariablesEditorListRow::SetWidgetCheckedState(const ECheckBoxState NewState,
+                                                           const bool bShouldUpdateHierarchyCheckedStates)
+{
+	WidgetCheckedState = NewState;
+}
+
+bool FConsoleVariablesEditorListRow::IsRowChecked() const
+{
+	return GetWidgetCheckedState() == ECheckBoxState::Checked;
+}
+
+bool FConsoleVariablesEditorListRow::ShouldBeVisible() const
+{
+	return (bDoesRowMatchSearchTerms && bDoesRowPassFilters) || HasVisibleChildren();
+}
+
+EVisibility FConsoleVariablesEditorListRow::GetDesiredVisibility() const
+{
+	return ShouldBeVisible() ? EVisibility::Visible : EVisibility::Collapsed;
+}
+
+const FString& FConsoleVariablesEditorListRow::GetCachedValue() const
+{
+	return CachedValue;
+}
+
+void FConsoleVariablesEditorListRow::SetCachedValue(const FString& NewCachedValue)
+{
+	CachedValue = NewCachedValue;
+}
+
+TArray<FConsoleVariablesEditorListRowPtr> FConsoleVariablesEditorListRow::GetSelectedTreeViewItems() const
+{
+	return ListViewPtr.Pin()->GetSelectedTreeViewItems();
+}
+
+FReply FConsoleVariablesEditorListRow::OnActionButtonClicked()
+{
+	const bool bIsListValid =
+		GetListViewPtr().IsValid() && GetListViewPtr().Pin()->GetListModelPtr().IsValid();
+
+	if (!ensure(bIsListValid))
+	{
+		return FReply::Unhandled();
+	}
+
+	const FConsoleVariablesEditorList::EConsoleVariablesEditorListMode ListMode =
+		bIsListValid
+			? GetListViewPtr().Pin()->GetListModelPtr().Pin()->GetListMode()
+			: FConsoleVariablesEditorList::EConsoleVariablesEditorListMode::Preset;
+
+	const bool bIsGlobalSearch =
+		ListMode == FConsoleVariablesEditorList::EConsoleVariablesEditorListMode::GlobalSearch;
+
+	FConsoleVariablesEditorModule& ConsoleVariablesEditorModule = FConsoleVariablesEditorModule::Get();
+
+	const FString& CommandName = GetCommandInfo().Pin()->Command;
+
+	const TObjectPtr<UConsoleVariablesAsset> EditableAsset = ConsoleVariablesEditorModule.GetPresetAsset();
+	check(EditableAsset);
+
+	if (bIsGlobalSearch)
+	{
+		FConsoleVariablesEditorAssetSaveData MatchingData;
+		if (!ConsoleVariablesEditorModule.GetPresetAsset()->FindSavedDataByCommandString(CommandName, MatchingData))
+		{
+			EditableAsset->AddOrSetConsoleObjectSavedData(
+				{
+					CommandName,
+					"",
+					ECheckBoxState::Checked
+				}
+			);
+		}
+		else
+		{
+			EditableAsset->RemoveConsoleVariable(CommandName);
+		}
+	}
+	else
+	{
+		ResetToStartupValueAndSource();
+
+		EditableAsset->RemoveConsoleVariable(CommandName);
+
+		ListViewPtr.Pin()->RebuildList();
+	}
+
+	return FReply::Handled();
+}
+
+void FConsoleVariablesEditorListRow::ResetToPresetValue()
+{
+	const FString& Value = GetPresetValue();
+	GetCommandInfo().Pin()->ExecuteCommand(Value);
+	SetCachedValue(Value);
+}
diff --git a/Engine/Plugins/Editor/ConsoleVariablesEditor/Source/ConsoleVariablesEditor/Private/Views/List/ConsoleVariablesEditorListRow.h b/Engine/Plugins/Editor/ConsoleVariablesEditor/Source/ConsoleVariablesEditor/Private/Views/List/ConsoleVariablesEditorListRow.h
new file mode 100644
--- /dev/null
+++ b/Engine/Plugins/Editor/ConsoleVariablesEditor/Source/ConsoleVariablesEditor/Private/Views/List/ConsoleVariablesEditorListRow.h
@@ -0,0 +1,163 @@
+// Copyright Epic Games, Inc. All Rights Reserved.
+
+#pragma once
+
+#include "CoreMinimal.h"
+#include "ConsoleVariablesEditorCommandInfo.h"
+#include "Widgets/Input/SCheckBox.h"
+
+class SConsoleVariablesEditorListValueInput;
+class SConsoleVariablesEditorList;
+
+struct FConsoleVariablesEditorListRow;
+typedef TSharedPtr<FConsoleVariablesEditorListRow> FConsoleVariablesEditorListRowPtr;
+
+struct FConsoleVariablesEditorListRow final : TSharedFromThis<FConsoleVariablesEditorListRow>
+{
+	enum EConsoleVariablesEditorListRowType
+	{
+		None,
+		HeaderRow, 
+		CommandGroup, // Group of commands or subgroups
+		SingleCommand
+	};
+
+	~FConsoleVariablesEditorListRow();
+
+	void FlushReferences();
+	
+	FConsoleVariablesEditorListRow(
+		const TWeakPtr<FConsoleVariablesEditorCommandInfo> InCommandInfo, const FString& InPresetValue, const EConsoleVariablesEditorListRowType InRowType, 
+		const ECheckBoxState StartingWidgetCheckboxState, const TSharedRef<SConsoleVariablesEditorList>& InListView, 
+		const int32 IndexInList, const TWeakPtr<FConsoleVariablesEditorListRow>& InDirectParentRow)
+	: CommandInfo(InCommandInfo)
+	, PresetValue(InPresetValue)
+	, RowType(InRowType)
+	, WidgetCheckedState(StartingWidgetCheckboxState)
+	, ListViewPtr(InListView)
+	, SortOrder(IndexInList)
+	, DirectParentRow(InDirectParentRow)
+	{
+		if (PresetValue.IsEmpty())
+		{
+			if (const auto Variable = CommandInfo.Pin()->GetConsoleVariablePtr())
+			{
+				PresetValue = Variable->GetString();
+			}
+		}
+
+		SetCachedValue(PresetValue);
+	}
+
+	[[nodiscard]] TWeakPtr<FConsoleVariablesEditorCommandInfo> GetCommandInfo() const;
+
+	[[nodiscard]] EConsoleVariablesEditorListRowType GetRowType() const;
+
+	[[nodiscard]] int32 GetChildDepth() const;
+	void SetChildDepth(const int32 InDepth);
+
+	[[nodiscard]] int32 GetSortOrder() const;
+	void SetSortOrder(const int32 InNewOrder);
+
+	TWeakPtr<FConsoleVariablesEditorListRow> GetDirectParentRow() const;
+	void SetDirectParentRow(const TWeakPtr<FConsoleVariablesEditorListRow>& InDirectParentRow);
+	
+	/* bHasGeneratedChildren must be true to get actual children. */
+	[[nodiscard]] const TArray<FConsoleVariablesEditorListRowPtr>& GetChildRows() const;
+	/* bHasGeneratedChildren must be true to get an accurate value. */
+	[[nodiscard]] int32 GetChildCount() const;
+	void SetChildRows(const TArray<FConsoleVariablesEditorListRowPtr>& InChildRows);
+	void AddToChildRows(const FConsoleVariablesEditorListRowPtr& InRow);
+	void InsertChildRowAtIndex(const FConsoleVariablesEditorListRowPtr& InRow, const int32 AtIndex = 0);
+
+	[[nodiscard]] bool GetIsTreeViewItemExpanded() const;
+	void SetIsTreeViewItemExpanded(const bool bNewExpanded);
+	
+	[[nodiscard]] bool GetShouldFlashOnScrollIntoView() const;
+	void SetShouldFlashOnScrollIntoView(const bool bNewShouldFlashOnScrollIntoView);
+
+	[[nodiscard]] bool GetShouldExpandAllChildren() const;
+	void SetShouldExpandAllChildren(const bool bNewShouldExpandAllChildren);
+
+	void ResetToStartupValueAndSource() const;
+	
+	[[nodiscard]] const FString& GetPresetValue() const;
+	void SetPresetValue(const FString& InPresetValue);
+
+	/*
+	 *Individual members of InTokens will be considered "AnyOf" or "OR" searches. If SearchTerms contains any individual member it will match.
+	 *Members will be tested for a space character (" "). If a space is found, a subsearch will be run.
+	 *This subsearch will be an "AllOf" or "AND" type search in which all strings, separated by a space, must be found in the search terms.
+	 */
+	bool MatchSearchTokensToSearchTerms(const TArray<FString> InTokens, ESearchCase::Type InSearchCase = ESearchCase::IgnoreCase);
+
+	/* This overload creates tokens from a string first, then calls ExecuteSearchOnChildNodes(const TArray<FString>& Tokens). */
+	void ExecuteSearchOnChildNodes(const FString& SearchString) const;
+	void ExecuteSearchOnChildNodes(const TArray<FString>& Tokens) const;
+
+	[[nodiscard]] bool GetDoesRowPassFilters() const;
+	void SetDoesRowPassFilters(const bool bPass);
+	
+	[[nodiscard]] bool DoesCurrentValueDifferFromPresetValue() const;
+	void SetDoesCurrentValueDifferFromPresetValue(const bool bNewValue);
+
+	[[nodiscard]] bool GetIsSelected() const;
+	void SetIsSelected(const bool bNewSelected);
+
+	[[nodiscard]] ECheckBoxState GetWidgetCheckedState() const;
+	void SetWidgetCheckedState(const ECheckBoxState NewState, const bool bShouldUpdateHierarchyCheckedStates = false);
+
+	[[nodiscard]] bool IsRowChecked() const;
+
+	[[nodiscard]] bool ShouldBeVisible() const;
+	[[nodiscard]] EVisibility GetDesiredVisibility() const;
+
+	[[nodiscard]] bool HasVisibleChildren() const
+	{
+		return false;
+	}
+
+	[[nodiscard]] TWeakPtr<SConsoleVariablesEditorList> GetListViewPtr() const
+	{
+		return ListViewPtr;
+	}
+	
+	[[nodiscard]] const FString& GetCachedValue() const;
+	void SetCachedValue(const FString& CachedValue);
+
+	[[nodiscard]] TArray<FConsoleVariablesEditorListRowPtr> GetSelectedTreeViewItems() const;
+
+	FReply OnActionButtonClicked();
+	
+	void ResetToPresetValue();
+
+private:
+
+	TWeakPtr<FConsoleVariablesEditorCommandInfo> CommandInfo;
+	FString PresetValue = "";
+	EConsoleVariablesEditorListRowType RowType = SingleCommand;
+	TArray<FConsoleVariablesEditorListRowPtr> ChildRows;
+
+	ECheckBoxState WidgetCheckedState = ECheckBoxState::Checked;
+	
+	TWeakPtr<SConsoleVariablesEditorList> ListViewPtr;
+	FString CachedValue = "";
+	
+	bool bIsTreeViewItemExpanded = false;
+	bool bShouldFlashOnScrollIntoView = false;
+
+	int32 ChildDepth = 0;
+
+	int32 SortOrder = -1;
+
+	bool bDoesRowMatchSearchTerms = true;
+	bool bDoesRowPassFilters = true;
+
+	bool bDoesCurrentValueDifferFromPresetValue = false;
+	
+	bool bIsSelected = false;
+	TWeakPtr<FConsoleVariablesEditorListRow> DirectParentRow;
+
+	// Used to expand all children on shift+click.
+	bool bShouldExpandAllChildren = false;
+};
diff --git a/Engine/Plugins/Editor/ConsoleVariablesEditor/Source/ConsoleVariablesEditor/Private/Views/List/SConsoleVariablesEditorList.cpp b/Engine/Plugins/Editor/ConsoleVariablesEditor/Source/ConsoleVariablesEditor/Private/Views/List/SConsoleVariablesEditorList.cpp
new file mode 100644
--- /dev/null
+++ b/Engine/Plugins/Editor/ConsoleVariablesEditor/Source/ConsoleVariablesEditor/Private/Views/List/SConsoleVariablesEditorList.cpp
@@ -0,0 +1,1224 @@
+// Copyright Epic Games, Inc. All Rights Reserved.
+
+#include "SConsoleVariablesEditorList.h"
+
+#include "ConsoleVariablesEditorList.h"
+#include "ConsoleVariablesEditorListFilters/ConsoleVariablesEditorListFilter_ModifiedVariables.h"
+#include "ConsoleVariablesEditorListFilters/ConsoleVariablesEditorListFilter_SourceText.h"
+#include "ConsoleVariablesEditorModule.h"
+#include "ConsoleVariablesEditorProjectSettings.h"
+#include "ConsoleVariablesEditorStyle.h"
+#include "SConsoleVariablesEditorListRow.h"
+#include "../Widgets/SConsoleVariablesEditorGlobalSearchToggle.h"
+
+#include "Algo/AllOf.h"
+#include "Algo/AnyOf.h"
+#include "Algo/Find.h"
+#include "Framework/MultiBox/MultiBoxBuilder.h"
+#include "Styling/StyleColors.h"
+#include "Widgets/Input/SCheckBox.h"
+#include "Widgets/Input/SComboBox.h"
+#include "Widgets/Input/SSearchBox.h"
+#include "Widgets/Layout/SWidgetSwitcher.h"
+#include "Widgets/Layout/SWrapBox.h"
+#include "Widgets/Text/SRichTextBlock.h"
+
+#define LOCTEXT_NAMESPACE "ConsoleVariablesEditor"
+
+const FName SConsoleVariablesEditorList::CustomSortOrderColumnName(TEXT("Order"));
+const FName SConsoleVariablesEditorList::CheckBoxColumnName(TEXT("Checkbox"));
+const FName SConsoleVariablesEditorList::VariableNameColumnName(TEXT("Name"));
+const FName SConsoleVariablesEditorList::ValueColumnName(TEXT("Value"));
+const FName SConsoleVariablesEditorList::SourceColumnName(TEXT("Source"));
+const FName SConsoleVariablesEditorList::ActionButtonColumnName(TEXT("Action"));
+
+void SConsoleVariablesEditorList::Construct(const FArguments& InArgs, TSharedRef<FConsoleVariablesEditorList> ListModel)
+{
+	ListModelPtr = ListModel;
+	
+	// Set Default Sorting info
+	ActiveSortingColumnName = CustomSortOrderColumnName;
+	ActiveSortingType = EColumnSortMode::Ascending;
+	
+	HeaderRow = SNew(SHeaderRow)
+				.CanSelectGeneratedColumn(true).Visibility(EVisibility::Visible);
+	
+	GenerateHeaderRow();
+
+	SetupFilters();
+	
+	ChildSlot
+	[
+		SNew(SVerticalBox)
+		
+		+ SVerticalBox::Slot()
+		.VAlign(VAlign_Top)
+		.AutoHeight()
+		.Padding(FMargin(8.f, 0.f, 8.f, 0.f))
+		[
+			SNew(SHorizontalBox)
+
+			+SHorizontalBox::Slot()
+			.HAlign(HAlign_Fill)
+			.Padding(0.f, 1.f, 0.f, 1.f)
+			[
+				SAssignNew(ListSearchBoxPtr, SSearchBox)
+				.HintText(LOCTEXT("ConsoleVariablesEditorList_SearchHintText", "Search..."))
+				.ToolTipText(LOCTEXT("ConsoleVariablesEditorList_TooltipText", "Search tracked variables, values, sources or help text"))
+				.OnTextChanged_Raw(this, &SConsoleVariablesEditorList::OnListViewSearchTextChanged)
+			]
+
+			// Global Search Button
+			+SHorizontalBox::Slot()
+			.VAlign(VAlign_Center)
+			.HAlign(HAlign_Right)
+			.AutoWidth()
+			.Padding(FMargin(8.f, 0.f, 0.f, 0.f))
+			[
+				SNew(SButton)
+				.ContentPadding(4.f)
+				.VAlign(VAlign_Center)
+				.OnClicked_Lambda([this]()
+				{
+					const FString SearchString = GetSearchStringFromSearchInputField().ToLower().TrimStartAndEnd();
+					ListSearchBoxPtr->SetText(FText::GetEmpty());
+					return TryEnterGlobalSearch(SearchString);
+				})
+				.ToolTipText(LOCTEXT("OpenInGlobalSearchButtonTooltip", "Search All Console Variables"))
+				[
+					SNew(SHorizontalBox)
+
+					+SHorizontalBox::Slot()
+					.Padding(0, 1, 4, 0)
+					.AutoWidth()
+					[
+						SNew(SImage)
+						.Image(FConsoleVariablesEditorStyle::Get().GetBrush("ConsoleVariables.GlobalSearch.Small"))
+					]
+
+					+SHorizontalBox::Slot()
+					.Padding(4.f, 1.f, 0.f, 0.f)
+					.AutoWidth()
+					[
+						SNew(STextBlock)
+						.Text(LOCTEXT("GlobalSearchButtonText","Search All"))
+					]
+				]
+			]
+
+			// Show Options
+			+SHorizontalBox::Slot()
+			.HAlign(HAlign_Right)
+			.VAlign(VAlign_Center)
+			.AutoWidth()
+			.Padding(8.f, 1.f, 0.f, 1.f)
+			[
+				SAssignNew( ViewOptionsComboButton, SComboButton )
+				.ContentPadding(0)
+				.ToolTipText(LOCTEXT("ShowOptions_Tooltip", "Show options to affect the visibility of items in the Console Variables Editor list"))
+				.ComboButtonStyle( FAppStyle::Get(), "SimpleComboButtonWithIcon" ) // Use the tool bar item style for this button
+				.OnGetMenuContent( this, &SConsoleVariablesEditorList::BuildShowOptionsMenu)
+				.HasDownArrow(false)
+				.ButtonContent()
+				[
+					SNew(SImage)
+					.ColorAndOpacity(FSlateColor::UseForeground())
+					.Image( FAppStyle::Get().GetBrush("Icons.Settings") )
+				]
+			]
+		]
+
+		+SVerticalBox::Slot()
+		.AutoHeight()
+		.Padding(FMargin(8.f, 4.f, 8.f, 4.f))
+		[
+			SAssignNew(GlobalSearchesHBox, SHorizontalBox)
+			.Visibility_Lambda([this]()
+			{
+				const bool bShouldBeVisible =
+					GlobalSearchesContainer.IsValid() && GlobalSearchesContainer->GetChildren()->Num() > 0;
+
+				return bShouldBeVisible ? EVisibility::SelfHitTestInvisible : EVisibility::Collapsed;
+			})
+
+			+SHorizontalBox::Slot()
+			.HAlign(HAlign_Fill)
+			[
+				SAssignNew(GlobalSearchesContainer, SWrapBox)
+				.InnerSlotPadding(FVector2d(6, 4))
+				.UseAllottedSize(true)
+			]
+
+			+SHorizontalBox::Slot()
+			.HAlign(HAlign_Right)
+			.VAlign(VAlign_Center)
+			.AutoWidth()
+			.Padding(8.f, 1.f, 2.f, 1.f)
+			[
+				// Remove Button
+				SAssignNew(RemoveGlobalSearchesButtonPtr, SCheckBox)
+				.Padding(0)
+				.ToolTipText(
+					LOCTEXT("RemoveGlobalSearchesButtonTooltip",
+					"Remove all global searches from the console variables editor."))
+				.Style(&FAppStyle::Get().GetWidgetStyle<FCheckBoxStyle>("ToggleButtonCheckbox"))
+				.ForegroundColor(FSlateColor::UseForeground())
+				.IsChecked(false)
+				.OnCheckStateChanged_Lambda([this] (ECheckBoxState NewCheckState)
+				{
+					ListModelPtr.Pin()->SetListMode(FConsoleVariablesEditorList::EConsoleVariablesEditorListMode::Preset);
+
+					GlobalSearchesContainer->ClearChildren();
+					CurrentGlobalSearches.Empty();
+
+					RebuildList();
+
+					RemoveGlobalSearchesButtonPtr->SetIsChecked(false);
+				})
+				[
+					SNew(SImage)
+					.Visibility(EVisibility::SelfHitTestInvisible)
+					.Image(FAppStyle::Get().GetBrush("Icons.X"))
+					.ColorAndOpacity(FSlateColor::UseForeground())
+				]
+			]
+		]
+
+		+ SVerticalBox::Slot()
+		[
+			SNew(SWidgetSwitcher)
+			.WidgetIndex_Lambda([this]()
+			{
+				return DoesTreeViewHaveVisibleChildren() ? 0 : 1;
+			})
+			
+			+ SWidgetSwitcher::Slot()
+			.HAlign(HAlign_Fill)
+			.Padding(2.0f, 2.0f, 2.0f, 2.0f)
+			[
+				SAssignNew(TreeViewPtr, STreeView<FConsoleVariablesEditorListRowPtr>)
+				.HeaderRow(HeaderRow)
+				.SelectionMode(ESelectionMode::Multi)
+				.OnSelectionChanged_Lambda([this] (const FConsoleVariablesEditorListRowPtr& Row, const ESelectInfo::Type SelectionType)
+				{
+					if(Row.IsValid())
+					{
+						Row->SetIsSelected(TreeViewPtr->GetSelectedItems().Contains(Row));
+					}
+				})
+				.TreeItemsSource(&VisibleTreeViewObjects)
+				.OnGenerateRow_Lambda([this](FConsoleVariablesEditorListRowPtr Row, const TSharedRef<STableViewBase>& OwnerTable)
+					{
+						check(Row.IsValid());
+					
+						return SNew(SConsoleVariablesEditorListRow, TreeViewPtr.ToSharedRef(), Row)
+								.Visibility_Raw(Row.Get(), &FConsoleVariablesEditorListRow::GetDesiredVisibility);
+					})
+				.OnGetChildren_Raw(this, &SConsoleVariablesEditorList::OnGetRowChildren)
+				.OnExpansionChanged_Raw(this, &SConsoleVariablesEditorList::OnRowChildExpansionChange, false)
+				.OnSetExpansionRecursive(this, &SConsoleVariablesEditorList::OnRowChildExpansionChange, true)
+			]
+
+			// For when no rows exist in view
+			+ SWidgetSwitcher::Slot()
+			.HAlign(HAlign_Fill)
+			.Padding(2.0f, 24.0f, 2.0f, 2.0f)
+			[
+				SNew(SRichTextBlock)
+				.DecoratorStyleSet(&FEditorStyle::Get())
+				.AutoWrapText(true)
+				.Justification(ETextJustify::Center)
+				.Text_Lambda([this]()
+				{
+					if (const TSharedPtr<FConsoleVariablesEditorList> ListModel = ListModelPtr.Pin())
+					{
+						if (ListModel->GetListMode() == FConsoleVariablesEditorList::EConsoleVariablesEditorListMode::GlobalSearch)
+						{
+							return LOCTEXT("ConsoleVariablesEditorList_NoList",
+								"No matching console variables found in Unreal Engine.\n\nCheck your search criteria.");
+						}
+					}
+					
+					return LOCTEXT("ConsoleVariablesEditorList_NoList",
+					"No matching console variables in your list.\n\nCheck your filter or <RichTextBlock.Bold>Search All</> console variables instead.");
+				})
+			]
+		]
+	];
+
+	EvaluateIfRowsPassFilters();
+}
+
+SConsoleVariablesEditorList::~SConsoleVariablesEditorList()
+{	
+	HeaderRow.Reset();
+	
+	ListSearchBoxPtr.Reset();
+	ViewOptionsComboButton.Reset();
+	GlobalSearchesHBox.Reset();
+	GlobalSearchesContainer.Reset();
+	CurrentGlobalSearches.Empty();
+	RemoveGlobalSearchesButtonPtr.Reset();
+	ListBoxContainerPtr.Reset();
+
+	FlushMemory(false);
+
+	ShowFilters.Reset();
+	TreeViewPtr.Reset();
+	VisibleTreeViewObjects.Reset();
+	LastPresetObjects.Reset();
+}
+
+FReply SConsoleVariablesEditorList::TryEnterGlobalSearch(const FString& SearchString)
+{
+	FReply ReturnValue = FReply::Unhandled();
+
+	// Can't enter global search if there are no active global searches or a search string from which to parse new searches
+	if (SearchString.IsEmpty() && CurrentGlobalSearches.Num() == 0)
+	{
+		UE_LOG(LogConsoleVariablesEditor, Log,
+			TEXT("%hs: Global search request is empty. Exiting Global Search."),
+			__FUNCTION__, *SearchString);
+
+		// Return to preset mode if in global search mode then rebuild list
+		if (ListModelPtr.Pin()->GetListMode() == FConsoleVariablesEditorList::EConsoleVariablesEditorListMode::GlobalSearch)
+		{
+			ListModelPtr.Pin()->SetListMode(FConsoleVariablesEditorList::EConsoleVariablesEditorListMode::Preset);
+			RebuildList();
+		}
+
+		return ReturnValue;
+	}
+
+	// Strings that already have associated buttons	
+	TArray<FString> ExistingButtonStrings;
+	
+	// All strings parsed from the search text
+	TArray<FString> OutTokens;
+	SearchString.ParseIntoArray(OutTokens, TEXT("|"), true);
+
+	// Get tokens from current searches. This step allows us to properly populate the asset with all matching commands
+	for (const TSharedRef<SConsoleVariablesEditorGlobalSearchToggle>& GlobalSearchButton : CurrentGlobalSearches)
+	{
+		const FString ButtonText = GlobalSearchButton->GetGlobalSearchText().ToString();
+		ExistingButtonStrings.Add(ButtonText);
+
+		// If the button text was explicitly typed into the search then set the existing button to be checked.
+		if (OutTokens.ContainsByPredicate([this, &ButtonText](const FString& Token)
+		{
+			return Token.Equals(ButtonText, ESearchCase::IgnoreCase);
+		}))
+		{
+			GlobalSearchButton->SetIsButtonChecked(true);
+		}
+
+		// If the button is checked, add its search text to the tokens array so it can be used to populate the matching command list
+		if (GlobalSearchButton->GetIsToggleChecked())
+		{
+			OutTokens.AddUnique(ButtonText);
+		}
+	}
+	
+	FConsoleVariablesEditorModule& ConsoleVariablesEditorModule = FConsoleVariablesEditorModule::Get();
+	
+	const bool bFoundMatches = ConsoleVariablesEditorModule.PopulateGlobalSearchAssetWithVariablesMatchingTokens(OutTokens);
+	{
+		// If we were in Preset mode before entering global search, cache the existing tree objects to maintain state
+		if (ListModelPtr.Pin()->GetListMode() == FConsoleVariablesEditorList::EConsoleVariablesEditorListMode::Preset)
+		{
+			LastPresetObjects = TreeViewRootObjects;
+			ListModelPtr.Pin()->SetListMode(FConsoleVariablesEditorList::EConsoleVariablesEditorListMode::GlobalSearch);
+		}
+
+		// Convert tokens to global search toggle buttons
+		for (const FString& TokenString : OutTokens)
+		{
+			// Only make new buttons when one doesn't exist for the current token
+			if (!ExistingButtonStrings.Contains(TokenString))
+			{
+				TSharedRef<SConsoleVariablesEditorGlobalSearchToggle> NewGlobalSearchButton =
+					SNew(SConsoleVariablesEditorGlobalSearchToggle, FText::FromString(TokenString))
+				.OnToggleClickedOnce_Lambda([this]()
+				{
+					return TryEnterGlobalSearch();
+				})
+				.OnToggleCtrlClicked(this, &SConsoleVariablesEditorList::HandleRemoveGlobalSearchToggleButton)
+				.OnToggleRightButtonClicked(this, &SConsoleVariablesEditorList::HandleRemoveGlobalSearchToggleButton);
+
+				CurrentGlobalSearches.Add(NewGlobalSearchButton);
+			}
+		}
+
+		// Put widgets in container
+		RefreshGlobalSearchWidgets();
+
+		ReturnValue = FReply::Handled();
+	}
+	
+	if (!bFoundMatches)
+	{
+		UE_LOG(LogConsoleVariablesEditor, Warning,
+			TEXT("%hs: Failed to find console variable objects with names containing search string %s"),
+			__FUNCTION__, *SearchString);
+	}
+
+	RebuildList();
+
+	return ReturnValue;
+}
+
+FReply SConsoleVariablesEditorList::HandleRemoveGlobalSearchToggleButton()
+{
+	CleanUpGlobalSearchesMarkedForDelete();
+	RefreshGlobalSearchWidgets();
+	return TryEnterGlobalSearch();
+}
+
+void SConsoleVariablesEditorList::CleanUpGlobalSearchesMarkedForDelete()
+{
+	for (int32 GlobalSearchItr = CurrentGlobalSearches.Num() - 1; GlobalSearchItr >= 0; GlobalSearchItr--)
+	{
+		if (CurrentGlobalSearches[GlobalSearchItr]->GetIsMarkedForDelete())
+		{
+			CurrentGlobalSearches.RemoveAt(GlobalSearchItr);
+		}
+	}
+}
+
+void SConsoleVariablesEditorList::RefreshGlobalSearchWidgets()
+{
+	GlobalSearchesContainer->ClearChildren();
+
+	for (const TSharedRef<SConsoleVariablesEditorGlobalSearchToggle>& GlobalSearchButton : CurrentGlobalSearches)
+	{	
+		GlobalSearchesContainer->AddSlot()
+		[
+			GlobalSearchButton
+		];
+	}
+}
+
+void SConsoleVariablesEditorList::RebuildList(const FString& InConsoleCommandToScrollTo, bool bShouldCacheValues)
+{
+	if (bShouldCacheValues)
+	{
+		CacheCurrentListItemData();
+	}
+
+	// Skip execution on load if we've cached the previous values
+	GenerateTreeView(!bShouldCacheValues);
+
+	if (bShouldCacheValues)
+	{
+		RestorePreviousListItemData();
+	}
+
+	RefreshList();
+
+	if (!InConsoleCommandToScrollTo.IsEmpty())
+	{
+		FConsoleVariablesEditorListRowPtr ScrollToItem = nullptr;
+
+		for (const FConsoleVariablesEditorListRowPtr& Item : TreeViewRootObjects)
+		{
+			if (Item->GetCommandInfo().Pin()->Command.Equals(InConsoleCommandToScrollTo))
+			{
+				ScrollToItem = Item;
+				break;
+			}
+		}
+
+		if (ScrollToItem.IsValid())
+		{
+			ScrollToItem->SetShouldFlashOnScrollIntoView(true);
+			TreeViewPtr->RequestScrollIntoView(ScrollToItem);
+		}
+	}
+}
+
+void SConsoleVariablesEditorList::RefreshList()
+{
+	if (TreeViewRootObjects.Num() > 0)
+	{
+		// Apply last search
+		ExecuteListViewSearchOnAllRows(GetSearchStringFromSearchInputField(), false);
+
+		// Enforce Sort
+		const FName& SortingName = GetActiveSortingColumnName();
+		ExecuteSort(SortingName, GetSortModeForColumn(SortingName), false);
+
+		// Show/Hide rows based on SetBy changes and filter settings
+		EvaluateIfRowsPassFilters(false);
+
+		// Refresh the header's check state
+		OnListItemCheckBoxStateChange(ECheckBoxState::Undetermined);
+
+		FindVisibleObjectsAndRequestTreeRefresh();
+	}
+}
+
+TArray<FConsoleVariablesEditorListRowPtr> SConsoleVariablesEditorList::GetSelectedTreeViewItems() const
+{
+	return TreeViewPtr->GetSelectedItems();
+}
+
+TArray<FConsoleVariablesEditorListRowPtr> SConsoleVariablesEditorList::GetTreeViewItems() const
+{
+	return TreeViewRootObjects;
+}
+
+void SConsoleVariablesEditorList::SetTreeViewItems(const TArray<FConsoleVariablesEditorListRowPtr>& InItems)
+{
+	TreeViewRootObjects = InItems;
+
+	TreeViewPtr->RequestListRefresh();
+}
+
+void SConsoleVariablesEditorList::UpdatePresetValuesForSave(const TObjectPtr<UConsoleVariablesAsset> InAsset) const
+{
+	TArray<FConsoleVariablesEditorAssetSaveData> NewSavedCommands;
+
+	const TArray<FConsoleVariablesEditorListRowPtr>& Items =
+		ListModelPtr.IsValid() && ListModelPtr.Pin()->GetListMode() ==
+			FConsoleVariablesEditorList::EConsoleVariablesEditorListMode::Preset ? TreeViewRootObjects : LastPresetObjects;
+	
+	for (const FConsoleVariablesEditorListRowPtr& Item : Items)
+	{
+		if (const TWeakPtr<FConsoleVariablesEditorCommandInfo> CommandInfo = Item->GetCommandInfo(); CommandInfo.IsValid())
+		{
+			NewSavedCommands.Add(
+				{
+					CommandInfo.Pin()->Command,
+					Item->GetCachedValue(),
+					Item->GetWidgetCheckedState()
+				}
+			);
+		}
+	}
+
+	InAsset->ReplaceSavedCommands(NewSavedCommands);
+}
+
+FString SConsoleVariablesEditorList::GetSearchStringFromSearchInputField() const
+{
+	return ensureAlwaysMsgf(ListSearchBoxPtr.IsValid(),
+		TEXT("%hs: ListSearchBoxPtr is not valid. Check to make sure it was created."), __FUNCTION__)
+	? ListSearchBoxPtr->GetText().ToString() : "";
+}
+
+void SConsoleVariablesEditorList::SetSearchStringInSearchInputField(const FString InSearchString) const
+{
+	if (ensureAlwaysMsgf(ListSearchBoxPtr.IsValid(),
+		TEXT("%hs: ListSearchBoxPtr is not valid. Check to make sure it was created."), __FUNCTION__))
+	{
+		ListSearchBoxPtr->SetText(FText::FromString(InSearchString));
+	}
+}
+
+void SConsoleVariablesEditorList::ExecuteListViewSearchOnAllRows(
+	const FString& SearchString, const bool bShouldRefreshAfterward)
+{
+	TArray<FString> Tokens;
+	
+	// unquoted search equivalent to a match-any-of search
+	SearchString.ParseIntoArray(Tokens, TEXT("|"), true);
+	
+	for (const TSharedPtr<FConsoleVariablesEditorListRow>& ChildRow : TreeViewRootObjects)
+	{
+		if (!ensure(ChildRow.IsValid()))
+		{
+			continue;
+		}
+		
+		const bool bGroupMatch = ChildRow->MatchSearchTokensToSearchTerms(Tokens);
+		
+		// If the group name matches then we pass in an empty string so all child nodes are visible.
+		// If the name doesn't match, then we need to evaluate each child.
+		ChildRow->ExecuteSearchOnChildNodes(bGroupMatch ? "" : SearchString);
+	}
+
+	if (bShouldRefreshAfterward)
+	{
+		FindVisibleObjectsAndRequestTreeRefresh();
+	}
+}
+
+bool SConsoleVariablesEditorList::DoesTreeViewHaveVisibleChildren() const
+{
+	if (TreeViewPtr.IsValid())
+	{
+		for (const TSharedPtr<FConsoleVariablesEditorListRow>& Header : TreeViewRootObjects)
+		{
+			const EVisibility HeaderVisibility = Header->GetDesiredVisibility();
+			
+			if (HeaderVisibility != EVisibility::Hidden && HeaderVisibility != EVisibility::Collapsed)
+			{
+				return true;
+			}
+		}
+	}
+	
+	return false;
+}
+
+void SConsoleVariablesEditorList::SetTreeViewItemExpanded(const TSharedPtr<FConsoleVariablesEditorListRow>& RowToExpand, const bool bNewExpansion) const
+{
+	if (TreeViewPtr.IsValid())
+	{
+		TreeViewPtr->SetItemExpansion(RowToExpand, bNewExpansion);
+	}
+}
+
+void SConsoleVariablesEditorList::SetAllListViewItemsCheckState(const ECheckBoxState InNewState)
+{
+	for (const TSharedPtr<FConsoleVariablesEditorListRow>& Row : TreeViewRootObjects)
+	{
+		Row->SetWidgetCheckedState(InNewState);
+	}
+}
+
+bool SConsoleVariablesEditorList::DoesListHaveCheckedMembers() const
+{
+	for (const TSharedPtr<FConsoleVariablesEditorListRow>& Row : TreeViewRootObjects)
+	{
+		if (Row->GetCommandInfo().Pin()->ObjectType ==
+			FConsoleVariablesEditorCommandInfo::EConsoleObjectType::Variable &&
+			Row->GetWidgetCheckedState() == ECheckBoxState::Checked)
+		{
+			return true;
+		}
+	}
+
+	return false;
+}
+
+bool SConsoleVariablesEditorList::DoesListHaveUncheckedMembers() const
+{
+	for (const TSharedPtr<FConsoleVariablesEditorListRow>& Row : TreeViewRootObjects)
+	{
+		if (Row->GetCommandInfo().Pin()->ObjectType ==
+			FConsoleVariablesEditorCommandInfo::EConsoleObjectType::Variable &&
+			Row->GetWidgetCheckedState() == ECheckBoxState::Unchecked)
+		{
+			return true;
+		}
+	}
+
+	return false;
+}
+
+void SConsoleVariablesEditorList::OnListItemCheckBoxStateChange(const ECheckBoxState InNewState)
+{
+	HeaderCheckBoxState = ECheckBoxState::Checked;
+
+	if (DoesListHaveUncheckedMembers())
+	{
+		HeaderCheckBoxState = ECheckBoxState::Unchecked;
+
+		if (DoesListHaveCheckedMembers())
+		{
+			HeaderCheckBoxState = ECheckBoxState::Undetermined;
+		}
+	}
+}
+
+void SConsoleVariablesEditorList::ToggleFilterActive(const FString& FilterName)
+{
+	if (const TSharedRef<IConsoleVariablesEditorListFilter>* Match =
+		Algo::FindByPredicate(ShowFilters,
+		[&FilterName](TSharedRef<IConsoleVariablesEditorListFilter> Comparator)
+		{
+			return Comparator->GetFilterName().Equals(FilterName);
+		}))
+	{
+		const TSharedRef<IConsoleVariablesEditorListFilter> Filter = *Match;
+		Filter->ToggleFilterActive();
+
+		EvaluateIfRowsPassFilters();
+	}
+}
+
+void SConsoleVariablesEditorList::EvaluateIfRowsPassFilters(const bool bShouldRefreshAfterward)
+{
+	// Separate filters by type
+	
+	TSet<TSharedRef<IConsoleVariablesEditorListFilter>> MatchAnyOfFilters;
+	TSet<TSharedRef<IConsoleVariablesEditorListFilter>> MatchAllOfFilters;
+
+	for (const TSharedRef<IConsoleVariablesEditorListFilter>& Filter : ShowFilters)
+	{
+		if (Filter->GetFilterMatchType() ==
+			IConsoleVariablesEditorListFilter::EConsoleVariablesEditorListFilterMatchType::MatchAll)
+		{
+			MatchAllOfFilters.Add(Filter);
+		}
+		else
+		{
+			MatchAnyOfFilters.Add(Filter);
+		}
+	}
+	
+	for (const FConsoleVariablesEditorListRowPtr& Row : TreeViewRootObjects)
+	{
+		if (Row.IsValid() && Row->GetRowType() == FConsoleVariablesEditorListRow::SingleCommand)
+		{
+			auto Projection = [&Row](const TSharedRef<IConsoleVariablesEditorListFilter>& Filter)
+			{
+				return Filter->GetIsFilterActive() ? Filter->DoesItemPassFilter(Row) : true;
+			};
+			
+			const bool bPassesAnyOf = Algo::AnyOf(MatchAnyOfFilters, Projection);
+			const bool bPassesAllOf = Algo::AllOf(MatchAllOfFilters, Projection);
+			
+			Row->SetDoesRowPassFilters(bPassesAnyOf && bPassesAllOf);
+		}
+	}
+
+	if (bShouldRefreshAfterward)
+	{
+		FindVisibleObjectsAndRequestTreeRefresh();
+	}
+}
+
+EColumnSortMode::Type SConsoleVariablesEditorList::GetSortModeForColumn(FName InColumnName) const
+{
+	EColumnSortMode::Type ColumnSortMode = EColumnSortMode::None;
+
+	if (GetActiveSortingColumnName().IsEqual(InColumnName))
+	{
+		ColumnSortMode = ActiveSortingType;
+	}
+
+	return ColumnSortMode;
+}
+
+void SConsoleVariablesEditorList::OnSortColumnCalled(EColumnSortPriority::Type Priority, const FName& ColumnName, EColumnSortMode::Type SortMode)
+{
+	ExecuteSort(ColumnName, CycleSortMode(ColumnName));
+}
+
+EColumnSortMode::Type SConsoleVariablesEditorList::CycleSortMode(const FName& InColumnName)
+{
+	// Custom handler for Custom Sort Order mode
+	if (InColumnName.IsEqual(CustomSortOrderColumnName))
+	{
+		ActiveSortingType = EColumnSortMode::Ascending;
+	}
+	else
+	{
+		const EColumnSortMode::Type PreviousColumnSortMode = GetSortModeForColumn(InColumnName);
+		ActiveSortingType = PreviousColumnSortMode ==
+			EColumnSortMode::Ascending ? EColumnSortMode::Descending : EColumnSortMode::Ascending;
+	}
+
+	ActiveSortingColumnName = InColumnName;
+	return ActiveSortingType;
+}
+
+void SConsoleVariablesEditorList::ExecuteSort(
+	const FName& InColumnName, const EColumnSortMode::Type InColumnSortMode, const bool bShouldRefreshAfterward)
+{	
+	if (InColumnName.IsEqual(CustomSortOrderColumnName))
+	{
+		TreeViewRootObjects.StableSort(SortByOrderAscending);
+	}
+	if (InColumnName.IsEqual(SourceColumnName))
+	{
+		TreeViewRootObjects.StableSort(
+			InColumnSortMode == EColumnSortMode::Ascending ? SortBySourceAscending : SortBySourceDescending);
+	}
+	if (InColumnName.IsEqual(VariableNameColumnName))
+	{
+		TreeViewRootObjects.StableSort(
+			InColumnSortMode == EColumnSortMode::Ascending ? SortByVariableNameAscending : SortByVariableNameDescending);
+	}
+
+	if (bShouldRefreshAfterward)
+	{
+		FindVisibleObjectsAndRequestTreeRefresh();
+	}
+}
+
+void SConsoleVariablesEditorList::SetSortOrder(const bool bShouldRefreshAfterward)
+{
+	for (int32 RowItr = 0; RowItr < TreeViewRootObjects.Num(); RowItr++)
+	{
+		const TSharedPtr<FConsoleVariablesEditorListRow>& ChildRow = TreeViewRootObjects[RowItr];
+		ChildRow->SetSortOrder(RowItr);
+	}
+
+	ExecuteSort(CustomSortOrderColumnName, CycleSortMode(CustomSortOrderColumnName), bShouldRefreshAfterward);
+}
+
+TSharedPtr<SHeaderRow> SConsoleVariablesEditorList::GenerateHeaderRow()
+{
+	check(HeaderRow);
+	HeaderRow->ClearColumns();
+
+	HeaderRow->AddColumn(
+		SHeaderRow::Column(CustomSortOrderColumnName)
+			.DefaultLabel(FText::FromString("#"))
+			.ToolTipText(LOCTEXT("ClickToSort","Click to sort"))
+			.HAlignHeader(EHorizontalAlignment::HAlign_Center)
+			.FillWidth(0.3f)
+			.ShouldGenerateWidget(true)
+			.SortMode_Raw(this, &SConsoleVariablesEditorList::GetSortModeForColumn, CustomSortOrderColumnName)
+			.OnSort_Raw(this, &SConsoleVariablesEditorList::OnSortColumnCalled)
+	);
+	
+	HeaderRow->AddColumn(
+		SHeaderRow::Column(CheckBoxColumnName)
+			.DefaultLabel(LOCTEXT("ConsoleVariablesEditorList_ConsoleVariableCheckboxHeaderText", "Checkbox"))
+			.HAlignHeader(EHorizontalAlignment::HAlign_Center)
+			.FixedWidth(25.f)
+			.ShouldGenerateWidget(true)
+			.HeaderContent()
+			[
+				SNew(SCheckBox)
+				.IsChecked_Lambda([this]()
+				{
+					return HeaderCheckBoxState;
+				})
+				.OnCheckStateChanged_Lambda([this] (const ECheckBoxState NewState)
+				{
+					HeaderCheckBoxState = NewState;
+					
+					for (const FConsoleVariablesEditorListRowPtr& Object : TreeViewRootObjects)
+					{
+						if (Object->GetCommandInfo().Pin()->ObjectType ==
+							FConsoleVariablesEditorCommandInfo::EConsoleObjectType::Variable)
+						{
+							Object->SetWidgetCheckedState(NewState);
+						}
+					}
+				})
+			]
+	);
+
+	HeaderRow->AddColumn(
+		SHeaderRow::Column(VariableNameColumnName)
+			.DefaultLabel(LOCTEXT("ConsoleVariablesEditorList_ConsoleVariableNameHeaderText", "Console Variable Name"))
+			.ToolTipText(LOCTEXT("ClickToSort","Click to sort"))
+			.HAlignHeader(EHorizontalAlignment::HAlign_Left)
+			.FillWidth(1.7f)
+			.ShouldGenerateWidget(true)
+			.SortMode_Raw(this, &SConsoleVariablesEditorList::GetSortModeForColumn, VariableNameColumnName)
+			.OnSort_Raw(this, &SConsoleVariablesEditorList::OnSortColumnCalled)
+	);
+
+	HeaderRow->AddColumn(
+		SHeaderRow::Column(ValueColumnName)
+			.DefaultLabel(LOCTEXT("ConsoleVariablesEditorList_ConsoleVariableValueHeaderText", "Value"))
+			.HAlignHeader(EHorizontalAlignment::HAlign_Left)
+			.ShouldGenerateWidget(true)
+			.FillWidth(0.8f)
+	);
+
+	HeaderRow->AddColumn(
+		SHeaderRow::Column(SourceColumnName)
+			.DefaultLabel(LOCTEXT("ConsoleVariablesEditorList_SourceHeaderText", "Source"))
+			.ToolTipText(LOCTEXT("ClickToSort","Click to sort"))
+			.HAlignHeader(EHorizontalAlignment::HAlign_Left)
+			.SortMode_Raw(this, &SConsoleVariablesEditorList::GetSortModeForColumn, SourceColumnName)
+			.OnSort_Raw(this, &SConsoleVariablesEditorList::OnSortColumnCalled)
+			);
+	
+	HeaderRow->AddColumn(
+		SHeaderRow::Column(ActionButtonColumnName)
+			.DefaultLabel(LOCTEXT("ConsoleVariablesEditorList_ConsoleVariableActionButtonHeaderText", "Action"))
+			.HAlignHeader(EHorizontalAlignment::HAlign_Center)
+			.FixedWidth(25.f)
+			.ShouldGenerateWidget(true)
+			.HeaderContent()
+			[
+				SNew(SBox)
+			]
+	);
+
+	return HeaderRow;
+}
+
+void SConsoleVariablesEditorList::SetupFilters()
+{
+	TArray<FString> SourceFilterTypes =
+	{
+		"Constructor",
+		"Scalability",
+		"Game Setting",
+		"Project Setting",
+		"System Settings ini",
+		"Device Profile",
+		"Game Override",
+		"Console Variables ini",
+		"Command line",
+		"Code",
+		"Console"
+	};
+
+	for (const FString& Type : SourceFilterTypes)
+	{
+		ShowFilters.Add(MakeShared<FConsoleVariablesEditorListFilter_SourceText>(Type));
+	}
+
+	// Add Show Only Modified filter
+	ShowFilters.Add(MakeShared<ConsoleVariablesEditorListFilter_ModifiedVariables>());
+}
+
+TSharedRef<SWidget> SConsoleVariablesEditorList::BuildShowOptionsMenu()
+{
+	FMenuBuilder ShowOptionsMenuBuilder = FMenuBuilder(true, nullptr);
+
+	ShowOptionsMenuBuilder.BeginSection("", LOCTEXT("ShowOptions_ShowSectionHeading", "Show"));
+	{
+		// Add show filters
+		auto AddFiltersLambda = [this, &ShowOptionsMenuBuilder](const TSharedRef<IConsoleVariablesEditorListFilter>& InFilter)
+		{
+			const FString& FilterName = InFilter->GetFilterName();
+			
+			ShowOptionsMenuBuilder.AddMenuEntry(
+			   InFilter->GetFilterButtonLabel(),
+			   InFilter->GetFilterButtonToolTip(),
+			   FSlateIcon(),
+			   FUIAction(
+				   FExecuteAction::CreateLambda(
+				   	[this, FilterName]()
+					   {
+						   ToggleFilterActive(FilterName);
+					   }
+					),
+				   FCanExecuteAction(),
+				   FIsActionChecked::CreateSP( InFilter, &IConsoleVariablesEditorListFilter::GetIsFilterActive )
+			   ),
+			   NAME_None,
+			   EUserInterfaceActionType::ToggleButton
+		   );
+		};
+
+		for (const TSharedRef<IConsoleVariablesEditorListFilter>& Filter : ShowFilters)
+		{
+			AddFiltersLambda(Filter);
+		}
+	}
+	ShowOptionsMenuBuilder.EndSection();
+	
+	ShowOptionsMenuBuilder.BeginSection("", LOCTEXT("ShowOptions_SortSectionHeading", "Sort"));
+	{
+		// Add commands
+
+		// Save this for later when folders are added
+		/*ShowOptionsMenuBuilder.AddMenuEntry(
+			LOCTEXT("CollapseAll", "Collapse All"),
+			LOCTEXT("ConsoleVariablesEditorList_CollapseAll_Tooltip", "Collapse all expanded actor groups in the Modified Actors list."),
+			FSlateIcon(),
+			FUIAction(FExecuteAction::CreateRaw(this, &SConsoleVariablesEditorList::SetAllGroupsCollapsed)),
+			NAME_None,
+			EUserInterfaceActionType::Button
+		);*/
+		
+		ShowOptionsMenuBuilder.AddMenuEntry(
+			LOCTEXT("SetSortOrder", "Set Sort Order"),
+			LOCTEXT("ConsoleVariablesEditorList_SetSortOrder_Tooltip", "Makes the current order of the variables list the saved order."),
+			FSlateIcon(),
+			FUIAction(FExecuteAction::CreateRaw(this, &SConsoleVariablesEditorList::SetSortOrder, true)),
+			NAME_None,
+			EUserInterfaceActionType::Button
+		);
+	}
+	ShowOptionsMenuBuilder.EndSection();
+	
+	ShowOptionsMenuBuilder.BeginSection("", LOCTEXT("ShowOptions_OptionsSectionHeading", "Options"));
+	{
+		ShowOptionsMenuBuilder.AddMenuEntry(
+			LOCTEXT("TrackAllVariableChanges", "Track All Variable Changes"),
+			LOCTEXT("ConsoleVariablesEditorList_TrackAllVariableChanges_Tooltip",
+				"When variables are changed outside the Console Variables Editor, this option will add the variables to the current preset. Does not apply to console commands like 'r.SetNearClipPlane' or 'stat fps'."),
+			FSlateIcon(),
+			FUIAction(
+				FExecuteAction::CreateLambda([]()
+				{
+					if (UConsoleVariablesEditorProjectSettings* ProjectSettingsPtr =
+						GetMutableDefault<UConsoleVariablesEditorProjectSettings>())
+					{
+						ProjectSettingsPtr->bAddAllChangedConsoleVariablesToCurrentPreset =
+							!ProjectSettingsPtr->bAddAllChangedConsoleVariablesToCurrentPreset;
+					}
+				}),
+			   FCanExecuteAction(),
+			   FIsActionChecked::CreateLambda( []()
+				{
+					if (const UConsoleVariablesEditorProjectSettings* ProjectSettingsPtr =
+						GetMutableDefault<UConsoleVariablesEditorProjectSettings>())
+					{
+						return ProjectSettingsPtr->bAddAllChangedConsoleVariablesToCurrentPreset;
+					}
+
+			   		return false;
+				})
+			),
+			NAME_None,
+			EUserInterfaceActionType::ToggleButton
+		);
+	}
+	ShowOptionsMenuBuilder.EndSection();
+
+	return ShowOptionsMenuBuilder.MakeWidget();
+}
+
+void SConsoleVariablesEditorList::FlushMemory(const bool bShouldKeepMemoryAllocated)
+{
+	if (bShouldKeepMemoryAllocated)
+	{
+		TreeViewRootObjects.Reset();
+	}
+	else
+	{
+		TreeViewRootObjects.Empty();
+	}
+}
+
+void SConsoleVariablesEditorList::SetAllGroupsCollapsed()
+{
+	if (TreeViewPtr.IsValid())
+	{
+		for (const FConsoleVariablesEditorListRowPtr& RootRow : TreeViewRootObjects)
+		{
+			if (!RootRow.IsValid())
+			{
+				continue;
+			}
+			
+			TreeViewPtr->SetItemExpansion(RootRow, false);
+			RootRow->SetIsTreeViewItemExpanded(false);
+		}
+	}
+}
+
+void SConsoleVariablesEditorList::OnListViewSearchTextChanged(const FText& Text)
+{
+	ExecuteListViewSearchOnAllRows(Text.ToString(), false);
+}
+
+void SConsoleVariablesEditorList::CacheCurrentListItemData()
+{
+	CachedCommandStates.Empty();
+
+	// We only want preset items, not global search
+	const TArray<FConsoleVariablesEditorListRowPtr>& ListItems =
+		ListModelPtr.Pin()->GetListMode() == FConsoleVariablesEditorList::EConsoleVariablesEditorListMode::Preset ?
+			TreeViewRootObjects : LastPresetObjects;
+
+	for (const TSharedPtr<FConsoleVariablesEditorListRow>& Item : ListItems)
+	{
+		if (const TSharedPtr<FConsoleVariablesEditorCommandInfo>& CommandInfo = Item->GetCommandInfo().Pin())
+		{
+			if (const IConsoleVariable* AsVariable = CommandInfo->GetConsoleVariablePtr())
+			{
+				CachedCommandStates.Add(
+					{
+						CommandInfo->Command,
+						AsVariable->GetString(),
+						Item->GetWidgetCheckedState()
+					}
+				);
+			}
+		}
+	}
+}
+
+void SConsoleVariablesEditorList::RestorePreviousListItemData()
+{
+	for (const TSharedPtr<FConsoleVariablesEditorListRow>& Item : TreeViewRootObjects)
+	{
+		if (const TSharedPtr<FConsoleVariablesEditorCommandInfo>& CommandInfo = Item->GetCommandInfo().Pin())
+		{
+			if (CommandInfo->GetConsoleVariablePtr())
+			{
+				const FString& CommandName = CommandInfo->Command;
+				
+				if (const FConsoleVariablesEditorAssetSaveData* Match = Algo::FindByPredicate(
+					CachedCommandStates,
+					[CommandName](const FConsoleVariablesEditorAssetSaveData& CachedData)
+					{
+						return CachedData.CommandName.Equals(CommandName);
+					}))
+				{
+					CommandInfo->ExecuteCommand((*Match).CommandValueAsString);
+					Item->SetWidgetCheckedState((*Match).CheckedState);
+				}
+			}
+		}
+	}
+
+	CachedCommandStates.Empty();
+}
+
+void SConsoleVariablesEditorList::GenerateTreeView(const bool bExecuteCommandsAsTheyAreLoaded)
+{	
+	if (!ensure(TreeViewPtr.IsValid()))
+	{
+		return;
+	}
+	
+	FlushMemory(true);
+
+	FConsoleVariablesEditorModule& ConsoleVariablesEditorModule = FConsoleVariablesEditorModule::Get();
+
+	TObjectPtr<UConsoleVariablesAsset> EditableAsset = nullptr;
+
+	switch (ListModelPtr.Pin()->GetListMode()) 
+	{
+		case FConsoleVariablesEditorList::EConsoleVariablesEditorListMode::Preset:
+			EditableAsset = ConsoleVariablesEditorModule.GetPresetAsset();
+			break;
+
+		case FConsoleVariablesEditorList::EConsoleVariablesEditorListMode::GlobalSearch:
+			EditableAsset = ConsoleVariablesEditorModule.GetGlobalSearchAsset();
+			break;
+
+		default:
+			checkNoEntry();
+	}
+	check(EditableAsset);
+
+	for (const FConsoleVariablesEditorAssetSaveData& SavedCommand : EditableAsset->GetSavedCommands())
+	{
+		// Get corresponding CommandInfo for tracking or make one if the command is non-value
+		TSharedPtr<FConsoleVariablesEditorCommandInfo> CommandInfo =
+			ConsoleVariablesEditorModule.FindCommandInfoByName(SavedCommand.CommandName).Pin();
+
+		if (!CommandInfo.IsValid())
+		{
+			CommandInfo = MakeShared<FConsoleVariablesEditorCommandInfo>(SavedCommand.CommandName);
+
+			FConsoleVariablesEditorModule::Get().AddConsoleObjectCommandInfoToMasterReference(
+				CommandInfo.ToSharedRef());
+		}
+		
+		if (CommandInfo.IsValid())
+		{
+			FConsoleVariablesEditorListRowPtr RowToAdd = nullptr;
+			
+			if (TSharedPtr<FConsoleVariablesEditorListRow>* MatchingRow = Algo::FindByPredicate(
+				LastPresetObjects,
+				[this, &CommandInfo](const FConsoleVariablesEditorListRowPtr& ExistingRow)
+				{
+					return ExistingRow->GetCommandInfo().IsValid() &&
+						ExistingRow->GetCommandInfo().Pin()->Command.Equals(CommandInfo->Command);
+				}))
+			{
+				RowToAdd = *MatchingRow;
+			}
+			else
+			{
+				const ECheckBoxState NewCheckedState =
+					SavedCommand.CheckedState == ECheckBoxState::Unchecked ?
+							ECheckBoxState::Unchecked : ECheckBoxState::Checked;
+
+				if (bExecuteCommandsAsTheyAreLoaded &&
+					ListModelPtr.Pin()->GetListMode() == FConsoleVariablesEditorList::EConsoleVariablesEditorListMode::Preset)
+				{
+					// If the row is checked and the saved value differs from the current value,
+					// execute the command with the saved value 
+					if (CommandInfo->ObjectType == FConsoleVariablesEditorCommandInfo::EConsoleObjectType::Command ||
+						(NewCheckedState == ECheckBoxState::Checked &&
+						CommandInfo->IsCurrentValueDifferentFromInputValue(SavedCommand.CommandValueAsString)))
+					{
+						if (!SavedCommand.CommandValueAsString.IsEmpty())
+						{
+							CommandInfo->ExecuteCommand(SavedCommand.CommandValueAsString);
+						}
+					}
+				}
+		
+				RowToAdd = 
+				   MakeShared<FConsoleVariablesEditorListRow>(
+						   CommandInfo, SavedCommand.CommandValueAsString,
+						   FConsoleVariablesEditorListRow::SingleCommand, 
+						   NewCheckedState, SharedThis(this), TreeViewRootObjects.Num(), nullptr);
+			}
+			
+			TreeViewRootObjects.Add(RowToAdd);
+		}
+	}
+
+	// Now clear out the last preset cache if the list is in preset mode
+	if (ListModelPtr.Pin()->GetListMode() == FConsoleVariablesEditorList::EConsoleVariablesEditorListMode::Preset)
+	{
+		LastPresetObjects.Empty();
+	}
+
+	TreeViewPtr->RequestTreeRefresh();
+}
+
+void SConsoleVariablesEditorList::FindVisibleTreeViewObjects()
+{
+	VisibleTreeViewObjects.Empty();
+
+	for (const TSharedPtr<FConsoleVariablesEditorListRow>& Row : TreeViewRootObjects)
+	{
+		if (Row->ShouldBeVisible())
+		{
+			VisibleTreeViewObjects.Add(Row);
+		}
+	}
+}
+
+void SConsoleVariablesEditorList::FindVisibleObjectsAndRequestTreeRefresh()
+{
+	FindVisibleTreeViewObjects();
+	TreeViewPtr->RequestTreeRefresh();
+}
+
+void SConsoleVariablesEditorList::OnGetRowChildren(FConsoleVariablesEditorListRowPtr Row, TArray<FConsoleVariablesEditorListRowPtr>& OutChildren) const
+{
+	if (Row.IsValid())
+	{
+		OutChildren = Row->GetChildRows();
+
+		if (Row->GetShouldExpandAllChildren())
+		{
+			SetChildExpansionRecursively(Row, true);
+			Row->SetShouldExpandAllChildren(false);
+		}
+	}
+}
+
+void SConsoleVariablesEditorList::OnRowChildExpansionChange(FConsoleVariablesEditorListRowPtr Row, const bool bIsExpanded, const bool bIsRecursive) const
+{
+	if (Row.IsValid())
+	{
+		if (bIsRecursive)
+		{
+			if (bIsExpanded)
+			{
+				if (Row->GetRowType() != FConsoleVariablesEditorListRow::HeaderRow)
+				{
+					Row->SetShouldExpandAllChildren(true);
+				}
+			}
+			else
+			{
+				SetChildExpansionRecursively(Row, bIsExpanded);
+			}
+		}
+		
+		TreeViewPtr->SetItemExpansion(Row, bIsExpanded);
+		Row->SetIsTreeViewItemExpanded(bIsExpanded);
+	}
+}
+
+void SConsoleVariablesEditorList::SetChildExpansionRecursively(const FConsoleVariablesEditorListRowPtr& InRow, const bool bNewIsExpanded) const
+{
+	if (InRow.IsValid())
+	{
+		for (const FConsoleVariablesEditorListRowPtr& Child : InRow->GetChildRows())
+		{
+			TreeViewPtr->SetItemExpansion(Child, bNewIsExpanded);
+			Child->SetIsTreeViewItemExpanded(bNewIsExpanded);
+
+			SetChildExpansionRecursively(Child, bNewIsExpanded);
+		}
+	}
+}
+
+#undef LOCTEXT_NAMESPACE
diff --git a/Engine/Plugins/Editor/ConsoleVariablesEditor/Source/ConsoleVariablesEditor/Private/Views/List/SConsoleVariablesEditorList.h b/Engine/Plugins/Editor/ConsoleVariablesEditor/Source/ConsoleVariablesEditor/Private/Views/List/SConsoleVariablesEditorList.h
new file mode 100644
--- /dev/null
+++ b/Engine/Plugins/Editor/ConsoleVariablesEditor/Source/ConsoleVariablesEditor/Private/Views/List/SConsoleVariablesEditorList.h
@@ -0,0 +1,214 @@
+// Copyright Epic Games, Inc. All Rights Reserved.
+
+#pragma once
+
+#include "ConsoleVariablesAsset.h"
+#include "ConsoleVariablesEditorListRow.h"
+#include "ConsoleVariablesEditorListFilters/IConsoleVariablesEditorListFilter.h"
+
+#include "Widgets/SCompoundWidget.h"
+#include "Widgets/Layout/SGridPanel.h"
+#include "Widgets/Views/SHeaderRow.h"
+#include "Widgets/Views/STreeView.h"
+
+class SConsoleVariablesEditorGlobalSearchToggle;
+class SWrapBox;
+class FConsoleVariablesEditorList;
+class UConsoleVariablesAsset;
+class SBox;
+class SComboButton;
+class SSearchBox;
+class SHeaderRow;
+
+class SConsoleVariablesEditorList final : public SCompoundWidget
+{
+
+public:
+
+	SLATE_BEGIN_ARGS(SConsoleVariablesEditorList)
+	{}
+
+	SLATE_END_ARGS()
+
+	void Construct(const FArguments& InArgs, TSharedRef<FConsoleVariablesEditorList> ListModel);
+
+	virtual ~SConsoleVariablesEditorList() override;
+
+	TWeakPtr<FConsoleVariablesEditorList> GetListModelPtr()
+	{
+		return ListModelPtr;
+	}
+
+	FReply TryEnterGlobalSearch(const FString& SearchString = "");
+	FReply HandleRemoveGlobalSearchToggleButton();
+	void CleanUpGlobalSearchesMarkedForDelete();
+	/** Remove all widgets from the global search container then recreate them */
+	void RefreshGlobalSearchWidgets();
+
+	/*
+	 * Regenerate the list items and refresh the list. Call when adding or removing variables.
+	 * @param bShouldCacheValues If true, the current list's current values will be cached and then restored when the list is rebuilt. Otherwise preset values will be used.
+	 */
+	void RebuildList(const FString& InConsoleCommandToScrollTo = "", bool bShouldCacheValues = true);
+
+	/**
+	 * Refresh filters and sorting.
+	 * Useful for when the list state has gone stale but the variable count has not changed.
+	 */
+	void RefreshList();
+
+	[[nodiscard]] TArray<FConsoleVariablesEditorListRowPtr> GetSelectedTreeViewItems() const;
+
+	[[nodiscard]] TArray<FConsoleVariablesEditorListRowPtr> GetTreeViewItems() const;
+	void SetTreeViewItems(const TArray<FConsoleVariablesEditorListRowPtr>& InItems);
+
+	[[nodiscard]] int32 GetTreeViewItemCount() const
+	{
+		return TreeViewRootObjects.Num();
+	}
+
+	/** Updates the saved values in a UConsoleVariablesAsset so that the command/value map can be saved to disk */
+	void UpdatePresetValuesForSave(const TObjectPtr<UConsoleVariablesAsset> InAsset) const;
+
+	FString GetSearchStringFromSearchInputField() const;
+	void SetSearchStringInSearchInputField(const FString InSearchString) const;
+	void ExecuteListViewSearchOnAllRows(const FString& SearchString, const bool bShouldRefreshAfterward = true);
+
+	bool DoesTreeViewHaveVisibleChildren() const;
+
+	void SetTreeViewItemExpanded(const TSharedPtr<FConsoleVariablesEditorListRow>& RowToExpand, const bool bNewExpansion) const;
+
+	void SetAllListViewItemsCheckState(const ECheckBoxState InNewState);
+
+	bool DoesListHaveCheckedMembers() const;
+
+	bool DoesListHaveUncheckedMembers() const;
+	
+	void OnListItemCheckBoxStateChange(const ECheckBoxState InNewState);
+
+	void ToggleFilterActive(const FString& FilterName);
+	void EvaluateIfRowsPassFilters(const bool bShouldRefreshAfterward = true);
+
+	// Sorting
+
+	const FName& GetActiveSortingColumnName() const
+	{
+		return ActiveSortingColumnName;
+	}
+	EColumnSortMode::Type GetSortModeForColumn(FName InColumnName) const;
+	void OnSortColumnCalled(EColumnSortPriority::Type Priority, const FName& ColumnName, EColumnSortMode::Type SortMode);
+	EColumnSortMode::Type CycleSortMode(const FName& InColumnName);
+	void ExecuteSort(
+		const FName& InColumnName, const EColumnSortMode::Type InColumnSortMode, const bool bShouldRefreshAfterward = true);
+	void ClearSorting()
+	{
+		ActiveSortingColumnName = NAME_None;
+		ActiveSortingType = EColumnSortMode::None;
+	}
+	void SetSortOrder(const bool bShouldRefreshAfterward = true);
+
+	// Column Names
+
+	static const FName CustomSortOrderColumnName;
+	static const FName CheckBoxColumnName;
+	static const FName VariableNameColumnName;
+	static const FName ValueColumnName;
+	static const FName SourceColumnName;
+	static const FName ActionButtonColumnName;
+
+private:
+
+	TWeakPtr<FConsoleVariablesEditorList> ListModelPtr;
+	
+	TSharedPtr<SHeaderRow> HeaderRow;
+	TSharedPtr<SHeaderRow> GenerateHeaderRow();
+	ECheckBoxState HeaderCheckBoxState = ECheckBoxState::Checked;
+
+	void SetupFilters();
+
+	TSharedRef<SWidget> BuildShowOptionsMenu();
+	
+	void FlushMemory(const bool bShouldKeepMemoryAllocated);
+	
+	void SetAllGroupsCollapsed();
+
+	// Search
+	
+	void OnListViewSearchTextChanged(const FText& Text);
+
+	TSharedPtr<SSearchBox> ListSearchBoxPtr;
+	TSharedPtr<SComboButton> ViewOptionsComboButton;
+	
+	/** Contains the GlobalSearchesContainer and remove button. Made a member in order to collapse it when GlobalSearchesContainer has no children. */
+	TSharedPtr<SHorizontalBox> GlobalSearchesHBox;
+	/** Contains the global search checkboxes only */
+	TSharedPtr<SWrapBox> GlobalSearchesContainer;
+	TArray<TSharedRef<SConsoleVariablesEditorGlobalSearchToggle>> CurrentGlobalSearches;
+	
+	TSharedPtr<SCheckBox> RemoveGlobalSearchesButtonPtr;
+	TSharedPtr<SBox> ListBoxContainerPtr;
+
+	//  Tree View Implementation
+
+	void CacheCurrentListItemData();
+	void RestorePreviousListItemData();
+	void GenerateTreeView(const bool bExecuteCommandsAsTheyAreLoaded = true);
+	void FindVisibleTreeViewObjects();
+	void FindVisibleObjectsAndRequestTreeRefresh();
+	
+	void OnGetRowChildren(FConsoleVariablesEditorListRowPtr Row, TArray<FConsoleVariablesEditorListRowPtr>& OutChildren) const;
+	void OnRowChildExpansionChange(FConsoleVariablesEditorListRowPtr Row, const bool bIsExpanded, const bool bIsRecursive = false) const;
+
+	void SetChildExpansionRecursively(const FConsoleVariablesEditorListRowPtr& InRow, const bool bNewIsExpanded) const;
+
+	TArray<TSharedRef<IConsoleVariablesEditorListFilter>> ShowFilters;
+
+	TSharedPtr<STreeView<FConsoleVariablesEditorListRowPtr>> TreeViewPtr;
+
+	/** All Tree view objects */
+	TArray<FConsoleVariablesEditorListRowPtr> TreeViewRootObjects;
+	/** Visible Tree view objects, after filters */
+	TArray<FConsoleVariablesEditorListRowPtr> VisibleTreeViewObjects;
+	/** When a global search is initiated we cache the tree objects shown in preset list mode
+	 * This way they can be recalled without rebuilding from the saved commands
+	 */
+	TArray<FConsoleVariablesEditorListRowPtr> LastPresetObjects;
+
+	/** A collection of list item data to save and restore when rebuilding the list. */
+	TArray<FConsoleVariablesEditorAssetSaveData> CachedCommandStates;
+
+	// Sorting
+
+	FName ActiveSortingColumnName = NAME_None;
+	EColumnSortMode::Type ActiveSortingType = EColumnSortMode::None;
+
+	TFunctionRef<bool(const FConsoleVariablesEditorListRowPtr&, const FConsoleVariablesEditorListRowPtr&)> SortByOrderAscending =
+		[](const FConsoleVariablesEditorListRowPtr& A, const FConsoleVariablesEditorListRowPtr& B)
+		{
+			return A->GetSortOrder() < B->GetSortOrder();
+		};
+
+	TFunctionRef<bool(const FConsoleVariablesEditorListRowPtr&, const FConsoleVariablesEditorListRowPtr&)> SortBySourceAscending =
+		[](const FConsoleVariablesEditorListRowPtr& A, const FConsoleVariablesEditorListRowPtr& B)
+		{
+			return A->GetCommandInfo().Pin()->GetSourceAsText().ToString() < B->GetCommandInfo().Pin()->GetSourceAsText().ToString();
+		};
+
+	TFunctionRef<bool(const FConsoleVariablesEditorListRowPtr&, const FConsoleVariablesEditorListRowPtr&)> SortBySourceDescending =
+		[](const FConsoleVariablesEditorListRowPtr& A, const FConsoleVariablesEditorListRowPtr& B)
+		{
+			return B->GetCommandInfo().Pin()->GetSourceAsText().ToString() < A->GetCommandInfo().Pin()->GetSourceAsText().ToString();
+		};
+	
+	TFunctionRef<bool(const FConsoleVariablesEditorListRowPtr&, const FConsoleVariablesEditorListRowPtr&)> SortByVariableNameAscending =
+		[](const FConsoleVariablesEditorListRowPtr& A, const FConsoleVariablesEditorListRowPtr& B)
+	{
+		return A->GetCommandInfo().Pin()->Command < B->GetCommandInfo().Pin()->Command;
+	};
+
+	TFunctionRef<bool(const FConsoleVariablesEditorListRowPtr&, const FConsoleVariablesEditorListRowPtr&)> SortByVariableNameDescending =
+		[](const FConsoleVariablesEditorListRowPtr& A, const FConsoleVariablesEditorListRowPtr& B)
+	{
+		return B->GetCommandInfo().Pin()->Command < A->GetCommandInfo().Pin()->Command;
+	};
+};
diff --git a/Engine/Plugins/Editor/ConsoleVariablesEditor/Source/ConsoleVariablesEditor/Private/Views/List/SConsoleVariablesEditorListRow.cpp b/Engine/Plugins/Editor/ConsoleVariablesEditor/Source/ConsoleVariablesEditor/Private/Views/List/SConsoleVariablesEditorListRow.cpp
new file mode 100644
--- /dev/null
+++ b/Engine/Plugins/Editor/ConsoleVariablesEditor/Source/ConsoleVariablesEditor/Private/Views/List/SConsoleVariablesEditorListRow.cpp
@@ -0,0 +1,649 @@
+// Copyright Epic Games, Inc. All Rights Reserved.
+
+#include "SConsoleVariablesEditorListRow.h"
+
+#include "ConsoleVariablesEditorList.h"
+#include "ConsoleVariablesEditorModule.h"
+#include "ConsoleVariablesEditorStyle.h"
+#include "SConsoleVariablesEditorListValueInput.h"
+#include "Views/Widgets/SConsoleVariablesEditorTooltipWidget.h"
+
+#include "Input/DragAndDrop.h"
+#include "Styling/AppStyle.h"
+#include "Styling/StyleColors.h"
+#include "Widgets/Input/SCheckBox.h"
+#include "Widgets/Layout/SBox.h"
+#include "Widgets/SBoxPanel.h"
+#include "Widgets/SOverlay.h"
+#include "Widgets/Layout/SScaleBox.h"
+#include "Widgets/Text/STextBlock.h"
+
+#define LOCTEXT_NAMESPACE "ConsoleVariablesEditor"
+
+class FConsoleVariablesListRowDragDropOp : public FDecoratedDragDropOp
+{
+public:
+	DRAG_DROP_OPERATOR_TYPE(FConsoleVariablesListRowDragDropOp, FDecoratedDragDropOp)
+
+	/** The item being dragged and dropped */
+	TArray<FConsoleVariablesEditorListRowPtr> DraggedItems;
+
+	/** Constructs a new drag/drop operation */
+	static TSharedRef<FConsoleVariablesListRowDragDropOp> New(const TArray<FConsoleVariablesEditorListRowPtr>& InItems)
+	{
+		check(InItems.Num() > 0);
+
+		TSharedRef<FConsoleVariablesListRowDragDropOp> Operation = MakeShareable(
+			new FConsoleVariablesListRowDragDropOp());
+
+		Operation->DraggedItems = InItems;
+
+		Operation->DefaultHoverIcon = FAppStyle::Get().GetBrush("Graph.ConnectorFeedback.Error");
+
+		// Set the display text and the transaction name based on whether we're dragging a single or multiple widgets
+		if (InItems.Num() == 1)
+		{
+			Operation->DefaultHoverText = FText::FromString(InItems[0]->GetCommandInfo().Pin()->Command);
+		}
+		else
+		{
+			Operation->DefaultHoverText =
+				FText::Format(
+					SConsoleVariablesEditorListRow::MultiDragFormatText,
+					FText::AsNumber(Operation->DraggedItems.Num())
+				);
+		}
+
+		Operation->Construct();
+
+		return Operation;
+	}
+};
+
+void SConsoleVariablesEditorListRow::Construct(
+	const FArguments& InArgs, const TSharedRef<STableViewBase>& InOwnerTable,
+	const TWeakPtr<FConsoleVariablesEditorListRow> InRow)
+{
+	check(InRow.IsValid());
+
+	Item = InRow;
+	const FConsoleVariablesEditorListRowPtr PinnedItem = Item.Pin();
+	const TSharedPtr<FConsoleVariablesEditorCommandInfo> PinnedCommand = PinnedItem->GetCommandInfo().Pin();
+
+	check(PinnedCommand.IsValid());
+
+	// Set up flash animation
+	FlashAnimation = FCurveSequence(0.f, FlashAnimationDuration, ECurveEaseFunction::QuadInOut);
+
+	SMultiColumnTableRow<FConsoleVariablesEditorListRowPtr>::Construct(
+		FSuperRowType::FArguments()
+		.Padding(1.0f)
+		.OnCanAcceptDrop(this, &SConsoleVariablesEditorListRow::HandleCanAcceptDrop)
+		.OnAcceptDrop(this, &SConsoleVariablesEditorListRow::HandleAcceptDrop)
+		.OnDragDetected(this, &SConsoleVariablesEditorListRow::HandleDragDetected)
+		.OnDragLeave(this, &SConsoleVariablesEditorListRow::HandleDragLeave),
+		InOwnerTable
+	);
+
+	if (PinnedItem->GetShouldFlashOnScrollIntoView())
+	{
+		FlashRow();
+
+		PinnedItem->SetShouldFlashOnScrollIntoView(false);
+	}
+}
+
+TSharedRef<SWidget> SConsoleVariablesEditorListRow::GenerateWidgetForColumn(const FName& InColumnName)
+{
+	const FConsoleVariablesEditorListRowPtr PinnedItem = Item.Pin();
+
+	const TSharedPtr<SWidget> CellWidget = GenerateCells(InColumnName, PinnedItem);
+
+	const TSharedRef<SImage> FlashImage = SNew(SImage)
+										.Image(new FSlateColorBrush(FStyleColors::White))
+										.Visibility_Raw(
+											this, &SConsoleVariablesEditorListRow::GetFlashImageVisibility)
+										.ColorAndOpacity_Raw(
+											this, &SConsoleVariablesEditorListRow::GetFlashImageColorAndOpacity);
+
+	FlashImages.Add(FlashImage);
+
+	return SNew(SOverlay)
+			.Visibility(EVisibility::SelfHitTestInvisible)
+
+			+ SOverlay::Slot()
+			[
+				FlashImage
+			]
+
+			+ SOverlay::Slot()
+			[
+				SNew(SBorder)
+					.HAlign(HAlign_Fill)
+					.VAlign(VAlign_Center)
+					.BorderImage(GetBorderImage(PinnedItem->GetRowType()))
+				[
+					CellWidget.ToSharedRef()
+				]
+			];
+}
+
+void SConsoleVariablesEditorListRow::OnMouseEnter(const FGeometry& MyGeometry, const FPointerEvent& MouseEvent)
+{
+	bIsHovered = true;
+
+	if (HoverableWidgetsPtr.IsValid())
+	{
+		HoverableWidgetsPtr->DetermineButtonImageAndTooltip();
+		HoverableWidgetsPtr->SetVisibility(EVisibility::SelfHitTestInvisible);
+	}
+
+	SMultiColumnTableRow<FConsoleVariablesEditorListRowPtr>::OnMouseEnter(MyGeometry, MouseEvent);
+}
+
+void SConsoleVariablesEditorListRow::OnMouseLeave(const FPointerEvent& MouseEvent)
+{
+	bIsHovered = false;
+
+	if (HoverableWidgetsPtr.IsValid())
+	{
+		HoverableWidgetsPtr->SetVisibility(EVisibility::Collapsed);
+	}
+
+	SMultiColumnTableRow<FConsoleVariablesEditorListRowPtr>::OnMouseLeave(MouseEvent);
+}
+
+SConsoleVariablesEditorListRow::~SConsoleVariablesEditorListRow()
+{
+	Item.Reset();
+	HoverToolTip.Reset();
+
+	FlashImages.Empty();
+
+	ValueChildInputWidget.Reset();
+
+	HoverableWidgetsPtr.Reset();
+}
+
+FReply SConsoleVariablesEditorListRow::HandleDragDetected(const FGeometry& MyGeometry, const FPointerEvent& MouseEvent)
+{
+	TArray<FConsoleVariablesEditorListRowPtr> DraggedItems = Item.Pin()->GetSelectedTreeViewItems();
+	TSharedRef<FConsoleVariablesListRowDragDropOp> Operation =
+		FConsoleVariablesListRowDragDropOp::New(DraggedItems);
+
+	return FReply::Handled().BeginDragDrop(Operation);
+}
+
+void SConsoleVariablesEditorListRow::HandleDragLeave(const FDragDropEvent& DragDropEvent)
+{
+	if (TSharedPtr<FConsoleVariablesListRowDragDropOp> Operation =
+		DragDropEvent.GetOperationAs<FConsoleVariablesListRowDragDropOp>())
+	{
+		Operation->ResetToDefaultToolTip();
+	}
+}
+
+TOptional<EItemDropZone> SConsoleVariablesEditorListRow::HandleCanAcceptDrop(const FDragDropEvent& DragDropEvent,
+                                                                             EItemDropZone DropZone,
+                                                                             FConsoleVariablesEditorListRowPtr
+                                                                             TargetItem)
+{
+	TSharedPtr<FConsoleVariablesListRowDragDropOp> Operation =
+		DragDropEvent.GetOperationAs<FConsoleVariablesListRowDragDropOp>();
+
+	if (!Operation.IsValid())
+	{
+		return TOptional<EItemDropZone>();
+	}
+
+	const bool bShouldEvaluateDrop = Item.IsValid() && Item.Pin()->GetListViewPtr().IsValid() &&
+		Item.Pin()->GetListViewPtr().Pin()->GetActiveSortingColumnName().IsEqual(
+			SConsoleVariablesEditorList::CustomSortOrderColumnName);
+
+	if (!bShouldEvaluateDrop)
+	{
+		Operation->SetToolTip(
+			LOCTEXT("SortByCustomOrderDrgDropWarning", "Sort by custom order (\"#\") to drag & drop"),
+			FAppStyle::Get().GetBrush("Graph.ConnectorFeedback.Error")
+		);
+
+		return TOptional<EItemDropZone>();
+	}
+
+	const bool bIsDropDenied = !TargetItem.IsValid() || Operation->DraggedItems.Num() < 1 ||
+		!TargetItem->GetCommandInfo().IsValid() || Operation->DraggedItems.Contains(TargetItem);
+
+	FString BoolAsString = bIsDropDenied ? "true" : "false";
+
+	if (bIsDropDenied)
+	{
+		Operation->ResetToDefaultToolTip();
+
+		return TOptional<EItemDropZone>();
+	}
+
+	FText ItemNameText = FText::FromString(Operation->DraggedItems[0]->GetCommandInfo().Pin()->Command);
+
+	if (Operation->DraggedItems.Num() > 1)
+	{
+		ItemNameText = FText::Format(MultiDragFormatText, FText::AsNumber(Operation->DraggedItems.Num()));
+	}
+
+	const FText DropPermittedText =
+		FText::Format(InsertFormatText,
+		              ItemNameText,
+		              DropZone == EItemDropZone::BelowItem ? BelowText : AboveText,
+		              FText::FromString(TargetItem->GetCommandInfo().Pin()->Command)
+		);
+
+	Operation->SetToolTip(
+		DropPermittedText,
+		FAppStyle::Get().GetBrush("Graph.ConnectorFeedback.OK")
+	);
+
+	// We have no behaviour yet for dropping one item onto another, so we'll treat it like we dropped it above
+	return DropZone == EItemDropZone::OntoItem ? EItemDropZone::AboveItem : DropZone;
+}
+
+FReply SConsoleVariablesEditorListRow::HandleAcceptDrop(const FDragDropEvent& DragDropEvent, EItemDropZone DropZone,
+                                                        FConsoleVariablesEditorListRowPtr TargetItem)
+{
+	TSharedPtr<FConsoleVariablesListRowDragDropOp> Operation =
+		DragDropEvent.GetOperationAs<FConsoleVariablesListRowDragDropOp>();
+
+	if (!Operation.IsValid())
+	{
+		return FReply::Unhandled();
+	}
+
+	const TSharedPtr<SConsoleVariablesEditorList> ListView = Item.Pin()->GetListViewPtr().Pin();
+
+	TArray<FConsoleVariablesEditorListRowPtr> DraggedItems = Operation->DraggedItems;
+
+	TArray<FConsoleVariablesEditorListRowPtr> AllTreeItemsCopy = ListView->GetTreeViewItems();
+
+	for (FConsoleVariablesEditorListRowPtr DraggedItem : DraggedItems)
+	{
+		if (!DraggedItem.IsValid() || !AllTreeItemsCopy.Contains(DraggedItem))
+		{
+			continue;
+		}
+
+		AllTreeItemsCopy.Remove(DraggedItem);
+	}
+
+	const int32 TargetIndex = AllTreeItemsCopy.IndexOfByKey(TargetItem);
+
+	if (TargetIndex > -1)
+	{
+		for (int32 ItemIndex = DraggedItems.Num() - 1; ItemIndex >= 0; ItemIndex--)
+		{
+			FConsoleVariablesEditorListRowPtr DraggedItem = DraggedItems[ItemIndex];
+
+			if (!DraggedItem.IsValid() || AllTreeItemsCopy.Contains(DraggedItem))
+			{
+				continue;
+			}
+
+			AllTreeItemsCopy.Insert(DraggedItem, DropZone == EItemDropZone::AboveItem ? TargetIndex : TargetIndex + 1);
+		}
+
+		ListView->SetTreeViewItems(AllTreeItemsCopy);
+		ListView->SetSortOrder();
+	}
+
+	return FReply::Handled();
+}
+
+void SConsoleVariablesEditorListRow::FlashRow()
+{
+	FlashAnimation.Play(this->AsShared());
+}
+
+EVisibility SConsoleVariablesEditorListRow::GetFlashImageVisibility() const
+{
+	return FlashAnimation.IsPlaying() ? EVisibility::SelfHitTestInvisible : EVisibility::Hidden;
+}
+
+FSlateColor SConsoleVariablesEditorListRow::GetFlashImageColorAndOpacity() const
+{
+	if (FlashAnimation.IsPlaying())
+	{
+		// This equation modulates the alpha into a parabolic curve 
+		const float Progress = FMath::Abs(FMath::Abs((FlashAnimation.GetLerp() - 0.5f) * 2) - 1);
+		return FLinearColor::LerpUsingHSV(FLinearColor::Transparent, FlashColor, Progress);
+	}
+
+	return FLinearColor::Transparent;
+}
+
+const FSlateBrush* SConsoleVariablesEditorListRow::GetBorderImage(
+	const FConsoleVariablesEditorListRow::EConsoleVariablesEditorListRowType InRowType)
+{
+	switch (InRowType)
+	{
+	case FConsoleVariablesEditorListRow::CommandGroup:
+		return FConsoleVariablesEditorStyle::Get().GetBrush("ConsoleVariablesEditor.CommandGroupBorder");
+
+	case FConsoleVariablesEditorListRow::HeaderRow:
+		return FConsoleVariablesEditorStyle::Get().GetBrush("ConsoleVariablesEditor.HeaderRowBorder");
+
+	default:
+		return FConsoleVariablesEditorStyle::Get().GetBrush("ConsoleVariablesEditor.DefaultBorder");
+	}
+}
+
+TSharedRef<SWidget> SConsoleVariablesEditorListRow::GenerateCells(const FName& InColumnName,
+                                                                  const TSharedPtr<FConsoleVariablesEditorListRow>
+                                                                  PinnedItem)
+{
+	if (InColumnName.IsEqual(SConsoleVariablesEditorList::CustomSortOrderColumnName))
+	{
+		return SNew(STextBlock)
+				.Visibility(EVisibility::SelfHitTestInvisible)
+				.Justification(ETextJustify::Center)
+				.Text_Lambda([this, PinnedItem]()
+		                       {
+			                       return FText::AsNumber(PinnedItem->GetSortOrder() + 1);
+		                       });
+	}
+	if (InColumnName.IsEqual(SConsoleVariablesEditorList::CheckBoxColumnName))
+	{
+		return SNew(SBox)
+				.Visibility(EVisibility::SelfHitTestInvisible)
+				.HAlign(HAlign_Center)
+		[
+			SNew(SCheckBox)
+					.IsChecked_Raw(this, &SConsoleVariablesEditorListRow::GetCheckboxState)
+					.OnCheckStateChanged_Raw(this, &SConsoleVariablesEditorListRow::OnCheckboxStateChange)
+					.Visibility(Item.Pin()->GetCommandInfo().Pin()->ObjectType ==
+					            FConsoleVariablesEditorCommandInfo::EConsoleObjectType::Variable
+						            ? EVisibility::Visible
+						            : EVisibility::Collapsed)
+		];
+	}
+	if (InColumnName.IsEqual(SConsoleVariablesEditorList::VariableNameColumnName))
+	{
+		if (!HoverToolTip.IsValid())
+		{
+			HoverToolTip = SConsoleVariablesEditorTooltipWidget::MakeTooltip(
+				PinnedItem->GetCommandInfo().Pin()->Command,
+				PinnedItem->GetCommandInfo().Pin()->GetHelpText());
+		}
+		return SNew(SBox)
+				.Visibility(EVisibility::SelfHitTestInvisible)
+				.HAlign(HAlign_Left)
+				.VAlign(VAlign_Center)
+				.Padding(FMargin(TextBlockLeftPadding, 0.f, 0.f, 0.f))
+				[
+					SNew(STextBlock)
+					.Visibility(EVisibility::Visible)
+					.Justification(ETextJustify::Left)
+					.Text(FText::FromString(PinnedItem->GetCommandInfo().Pin()->Command))
+					.ToolTip(HoverToolTip)
+				];
+	}
+	if (InColumnName.IsEqual(SConsoleVariablesEditorList::ValueColumnName))
+	{
+		return GenerateValueCellWidget(PinnedItem);
+	}
+	if (InColumnName.IsEqual(SConsoleVariablesEditorList::SourceColumnName))
+	{
+		return SNew(SBox)
+				.Visibility(EVisibility::HitTestInvisible)
+				.HAlign(HAlign_Left)
+				.VAlign(VAlign_Center)
+				.Padding(FMargin(TextBlockLeftPadding, 0.f, 0.f, 0.f))
+				[
+					SNew(STextBlock)
+					.Text_Lambda([this]()
+	                {
+		                return Item.Pin()->GetCommandInfo().Pin()->GetSourceAsText();
+	                })
+				];
+	}
+	if (InColumnName.IsEqual(SConsoleVariablesEditorList::ActionButtonColumnName))
+	{
+		return SAssignNew(HoverableWidgetsPtr, SConsoleVariablesEditorListRowHoverWidgets, Item)
+				.Visibility(EVisibility::Collapsed);
+	}
+
+	return SNullWidget::NullWidget;
+}
+
+ECheckBoxState SConsoleVariablesEditorListRow::GetCheckboxState() const
+{
+	//Non-variable types always checked
+	if (Item.Pin()->GetCommandInfo().Pin()->ObjectType !=
+		FConsoleVariablesEditorCommandInfo::EConsoleObjectType::Variable)
+	{
+		return ECheckBoxState::Checked;
+	}
+
+	if (ensure(Item.IsValid()))
+	{
+		return Item.Pin()->GetWidgetCheckedState();
+	}
+
+	return ECheckBoxState::Checked;
+}
+
+void SConsoleVariablesEditorListRow::OnCheckboxStateChange(const ECheckBoxState InNewState) const
+{
+	if (const TSharedPtr<FConsoleVariablesEditorListRow> PinnedItem = Item.Pin())
+	{
+		PinnedItem->SetWidgetCheckedState(InNewState, true);
+
+		if (PinnedItem->GetRowType() == FConsoleVariablesEditorListRow::SingleCommand)
+		{
+			if (PinnedItem->IsRowChecked())
+			{
+				PinnedItem->GetCommandInfo().Pin()->ExecuteCommand(PinnedItem->GetCachedValue());
+
+				FConsoleVariablesEditorModule::Get().SendMultiUserConsoleVariableChange(
+					PinnedItem->GetCommandInfo().Pin()->Command, PinnedItem->GetCachedValue());
+			}
+			else
+			{
+				PinnedItem->ResetToStartupValueAndSource();
+
+				FConsoleVariablesEditorModule::Get().SendMultiUserConsoleVariableChange(
+					PinnedItem->GetCommandInfo().Pin()->Command,
+					PinnedItem->GetCommandInfo().Pin()->StartupValueAsString
+				);
+			}
+		}
+
+		if (const TWeakPtr<SConsoleVariablesEditorList> ListView = PinnedItem->GetListViewPtr(); ListView.IsValid())
+		{
+			const TSharedPtr<SConsoleVariablesEditorList> PinnedListView = ListView.Pin();
+			PinnedListView->OnListItemCheckBoxStateChange(InNewState);
+		}
+	}
+}
+
+TSharedRef<SWidget> SConsoleVariablesEditorListRow::GenerateValueCellWidget(
+	const TSharedPtr<FConsoleVariablesEditorListRow> PinnedItem)
+{
+	if (PinnedItem->GetCommandInfo().IsValid())
+	{
+		ValueChildInputWidget = SConsoleVariablesEditorListValueInput::GetInputWidget(Item);
+		const TSharedRef<SHorizontalBox> FinalValueWidget =
+			SNew(SHorizontalBox)
+			.ToolTipText_Lambda([this, PinnedItem]()
+			{
+				if (const TSharedPtr<FConsoleVariablesEditorCommandInfo> PinnedCommand = PinnedItem->GetCommandInfo().
+						Pin();
+					PinnedCommand.IsValid())
+				{
+					return FText::Format(
+						ValueWidgetToolTipFormatText,
+						FText::FromString(PinnedItem->GetCachedValue()),
+						FText::FromString(PinnedItem->GetPresetValue()),
+						FText::FromString(PinnedCommand->StartupValueAsString),
+						FConsoleVariablesEditorCommandInfo::ConvertConsoleVariableSetByFlagToText(
+							PinnedCommand->StartupSource)
+					);
+				}
+
+				return FText::GetEmpty();
+			});
+
+		FinalValueWidget->AddSlot()
+		                .VAlign(VAlign_Center)
+		                .Padding(FMargin(2, 0))
+		[
+			ValueChildInputWidget.ToSharedRef()
+		];
+
+		FinalValueWidget->AddSlot()
+		                .AutoWidth()
+		                .VAlign(VAlign_Center)
+		                .Padding(FMargin(2, 0))
+		[
+			SNew(SButton)
+			.IsFocusable(false)
+			.ToolTipText(FText::Format(RevertButtonFormatText, FText::FromString(PinnedItem->GetPresetValue())))
+			.ButtonStyle(&FAppStyle::Get().GetWidgetStyle<FButtonStyle>("NoBorder"))
+			.ContentPadding(0)
+			.Visibility_Lambda([this, PinnedItem]()
+             {
+	             if (PinnedItem.IsValid() && PinnedItem->GetRowType() ==
+		             FConsoleVariablesEditorListRow::SingleCommand)
+	             {
+		             const TSharedPtr<FConsoleVariablesEditorCommandInfo> CommandInfo =
+			             PinnedItem->GetCommandInfo().Pin();
+		            
+		            if (CommandInfo.IsValid() &&
+			            CommandInfo->ObjectType == FConsoleVariablesEditorCommandInfo::EConsoleObjectType::Variable)
+		            {
+			            PinnedItem->SetDoesCurrentValueDifferFromPresetValue(
+							 CommandInfo->IsCurrentValueDifferentFromInputValue(
+								 PinnedItem->GetPresetValue()));
+
+						 return PinnedItem->IsRowChecked() && PinnedItem->
+								DoesCurrentValueDifferFromPresetValue()
+									? EVisibility::Visible
+									: EVisibility::Collapsed;
+		            }
+	             }
+
+	             return EVisibility::Collapsed;
+             })
+			.OnClicked_Lambda([this, PinnedItem]()
+             {
+	             PinnedItem->ResetToPresetValue();
+
+	             return FReply::Handled();
+             })
+			[
+				SNew(SImage)
+				.Image(FAppStyle::Get().GetBrush("PropertyWindow.DiffersFromDefault"))
+				.ColorAndOpacity(FSlateColor::UseForeground())
+			]
+		];
+
+		return FinalValueWidget;
+	}
+
+	return SNullWidget::NullWidget;
+}
+
+void SConsoleVariablesEditorListRowHoverWidgets::Construct(const FArguments& InArgs,
+                                                           const TWeakPtr<FConsoleVariablesEditorListRow> InRow)
+{
+	check(InRow.IsValid());
+
+	Item = InRow;
+
+	const FText ButtonTooltip = LOCTEXT("RemoveCvarTooltip",
+								  "Remove cvar from this list and reset its value to the startup value.");
+	const FSlateBrush* ButtonImage = FAppStyle::Get().GetBrush("Icons.Delete");
+
+	ChildSlot
+	[
+		// Action Button
+		SAssignNew(ActionButtonPtr, SButton)
+		.ButtonStyle(&FAppStyle::Get().GetWidgetStyle<FButtonStyle>("NoBorder"))
+		.Visibility(EVisibility::Visible)
+		.ToolTipText(ButtonTooltip)
+		.ButtonColorAndOpacity(FStyleColors::Transparent)
+		.ContentPadding(0.f)
+		.OnClicked_Lambda([this]()
+		{
+			SetVisibility(EVisibility::Collapsed);
+			FReply Reply = Item.Pin()->OnActionButtonClicked();
+			DetermineButtonImageAndTooltip();
+			SetVisibility(EVisibility::SelfHitTestInvisible);
+			return Reply;
+		})
+		[
+			SNew(SScaleBox)
+			[
+				SAssignNew(ActionButtonImage, SImage)
+				.Visibility(EVisibility::SelfHitTestInvisible)
+				.Image(ButtonImage)
+				.ColorAndOpacity(FSlateColor::UseForeground())
+			]
+		]
+	];
+}
+
+void SConsoleVariablesEditorListRowHoverWidgets::DetermineButtonImageAndTooltip()
+{
+	const bool bIsListValid =
+		Item.Pin()->GetListViewPtr().IsValid() && Item.Pin()->GetListViewPtr().Pin()->GetListModelPtr().IsValid();
+
+	const FConsoleVariablesEditorList::EConsoleVariablesEditorListMode ListMode =
+		bIsListValid
+			? Item.Pin()->GetListViewPtr().Pin()->GetListModelPtr().Pin()->GetListMode()
+			: FConsoleVariablesEditorList::EConsoleVariablesEditorListMode::Preset;
+
+	const bool bIsGlobalSearch =
+		ListMode == FConsoleVariablesEditorList::EConsoleVariablesEditorListMode::GlobalSearch;
+
+	if (bIsGlobalSearch)
+	{
+		FText ButtonTooltip = LOCTEXT("AddCvarToPresetTooltip", "Add this cvar to your current preset.");
+		const FSlateBrush* ButtonImage = FConsoleVariablesEditorStyle::Get().GetBrush("ConsoleVariables.Favorite.Outline.Small");
+		
+		const FConsoleVariablesEditorModule& ConsoleVariablesEditorModule = FConsoleVariablesEditorModule::Get();
+		
+		const FString& CommandName = Item.Pin()->GetCommandInfo().Pin()->Command;
+		FConsoleVariablesEditorAssetSaveData MatchingData;
+
+		// If the item exists in the current preset already
+		if (ConsoleVariablesEditorModule.GetPresetAsset()->FindSavedDataByCommandString(CommandName, MatchingData))
+		{
+			ButtonTooltip = LOCTEXT("RemoveCvarFromPresetTooltip", "Remove this cvar from your current preset.");
+			ButtonImage = FAppStyle::Get().GetBrush("Icons.Star");
+		}
+
+		ActionButtonPtr->SetToolTipText(ButtonTooltip);
+		ActionButtonImage->SetImage(ButtonImage);
+	}
+}
+
+void SConsoleVariablesEditorListRowHoverWidgets::OnMouseEnter(const FGeometry& MyGeometry,
+                                                              const FPointerEvent& MouseEvent)
+{
+	SCompoundWidget::OnMouseEnter(MyGeometry, MouseEvent);
+
+	ActionButtonPtr->SetBorderBackgroundColor(FLinearColor(0.f, 0.f, 0.f, 0.4f));
+}
+
+void SConsoleVariablesEditorListRowHoverWidgets::OnMouseLeave(const FPointerEvent& MouseEvent)
+{
+	SCompoundWidget::OnMouseLeave(MouseEvent);
+
+	ActionButtonPtr->SetBorderBackgroundColor(FLinearColor(0.f, 0.f, 0.f, 0.f));
+}
+
+SConsoleVariablesEditorListRowHoverWidgets::~SConsoleVariablesEditorListRowHoverWidgets()
+{
+	Item.Reset();
+	
+	ActionButtonPtr.Reset();
+}
+
+#undef LOCTEXT_NAMESPACE
diff --git a/Engine/Plugins/Editor/ConsoleVariablesEditor/Source/ConsoleVariablesEditor/Private/Views/List/SConsoleVariablesEditorListRow.h b/Engine/Plugins/Editor/ConsoleVariablesEditor/Source/ConsoleVariablesEditor/Private/Views/List/SConsoleVariablesEditorListRow.h
new file mode 100644
--- /dev/null
+++ b/Engine/Plugins/Editor/ConsoleVariablesEditor/Source/ConsoleVariablesEditor/Private/Views/List/SConsoleVariablesEditorListRow.h
@@ -0,0 +1,118 @@
+// Copyright Epic Games, Inc. All Rights Reserved.
+
+#pragma once
+
+#include "ConsoleVariablesEditorListRow.h"
+#include "SConsoleVariablesEditorList.h"
+
+#include "CoreMinimal.h"
+#include "DragAndDrop/DecoratedDragDropOp.h"
+#include "Widgets/Input/SButton.h"
+#include "Widgets/Input/SNumericEntryBox.h"
+#include "Widgets/Layout/SBorder.h"
+#include "Widgets/SCompoundWidget.h"
+#include "Widgets/IToolTip.h"
+
+class SConsoleVariablesEditorListValueInput;
+class SConsoleVariablesEditorListRowHoverWidgets;
+
+class SConsoleVariablesEditorListRow : public SMultiColumnTableRow<FConsoleVariablesEditorListRowPtr>
+{
+public:
+	
+	SLATE_BEGIN_ARGS(SConsoleVariablesEditorListRow)
+	{}
+
+	SLATE_END_ARGS()
+	
+	void Construct(const FArguments& InArgs, const TSharedRef<STableViewBase>& InOwnerTable, const TWeakPtr<FConsoleVariablesEditorListRow> InRow);
+
+	virtual TSharedRef<SWidget> GenerateWidgetForColumn(const FName& InColumnName) override;
+
+	// Begin SWidget
+	virtual void OnMouseEnter(const FGeometry& MyGeometry, const FPointerEvent& MouseEvent) override;
+	virtual void OnMouseLeave(const FPointerEvent& MouseEvent) override;
+	// End SWidget
+	
+	virtual ~SConsoleVariablesEditorListRow() override;	
+
+	FReply HandleDragDetected(const FGeometry& MyGeometry, const FPointerEvent& MouseEvent);
+	void HandleDragLeave(const FDragDropEvent& DragDropEvent);
+	TOptional<EItemDropZone> HandleCanAcceptDrop(const FDragDropEvent& DragDropEvent, EItemDropZone DropZone, FConsoleVariablesEditorListRowPtr TargetItem);
+	FReply HandleAcceptDrop(const FDragDropEvent& DragDropEvent, EItemDropZone DropZone, FConsoleVariablesEditorListRowPtr TargetItem);
+	
+	void FlashRow();
+	
+	EVisibility GetFlashImageVisibility() const;
+	FSlateColor GetFlashImageColorAndOpacity() const;
+
+	static const FSlateBrush* GetBorderImage(const FConsoleVariablesEditorListRow::EConsoleVariablesEditorListRowType InRowType);
+
+	TSharedRef<SWidget> GenerateCells(const FName& InColumnName, const TSharedPtr<FConsoleVariablesEditorListRow> PinnedItem);
+
+	ECheckBoxState GetCheckboxState() const;
+	void OnCheckboxStateChange(const ECheckBoxState InNewState) const;
+
+	TSharedRef<SWidget> GenerateValueCellWidget(const TSharedPtr<FConsoleVariablesEditorListRow> PinnedItem);
+
+	#define LOCTEXT_NAMESPACE "ConsoleVariablesEditor"
+
+	const FText ValueWidgetToolTipFormatText = LOCTEXT("ValueWidgetToolTipFormatText", "Custom Value: {0}\nPreset Value: {1}\nStartup Value: {2} (Set By {3})");
+	const FText RevertButtonFormatText = LOCTEXT("RevertButtonFormatText", "Reset to Preset Value: {0}");
+
+	const FText InsertFormatText = LOCTEXT("InsertAboveFormatText", "Insert {0} {1} {2}");
+	const FText AboveText = LOCTEXT("AboveListItem", "above");
+	const FText BelowText = LOCTEXT("BelowListItem", "below");
+	static const inline FText MultiDragFormatText = LOCTEXT("MultiDragFormatText", "{0} Items");
+
+	#undef LOCTEXT_NAMESPACE
+
+private:
+	
+	TWeakPtr<FConsoleVariablesEditorListRow> Item;
+	
+	TSharedPtr<IToolTip> HoverToolTip;
+
+	TArray<TSharedPtr<SImage>> FlashImages;
+
+	TSharedPtr<SConsoleVariablesEditorListValueInput> ValueChildInputWidget;
+	
+	TSharedPtr<SConsoleVariablesEditorListRowHoverWidgets> HoverableWidgetsPtr;
+
+	FCurveSequence FlashAnimation;
+
+	const float FlashAnimationDuration = 0.75f;
+	const FLinearColor FlashColor = FLinearColor::White;
+
+	/** The offset applied to text widgets so that the text aligns with the column header text */
+	float TextBlockLeftPadding = 3.0f;
+	
+	bool bIsHovered = false;
+};
+
+class SConsoleVariablesEditorListRowHoverWidgets : public SCompoundWidget
+{
+public:
+	
+	SLATE_BEGIN_ARGS(SConsoleVariablesEditorListRowHoverWidgets)
+	{}
+
+	SLATE_END_ARGS()
+	
+	void Construct(const FArguments& InArgs, const TWeakPtr<FConsoleVariablesEditorListRow> InRow);
+
+	void DetermineButtonImageAndTooltip();
+
+	virtual void OnMouseEnter(const FGeometry& MyGeometry, const FPointerEvent& MouseEvent) override;
+
+	virtual void OnMouseLeave(const FPointerEvent& MouseEvent) override;
+
+	virtual ~SConsoleVariablesEditorListRowHoverWidgets() override;
+
+private:
+	
+	TWeakPtr<FConsoleVariablesEditorListRow> Item;
+	
+	TSharedPtr<SButton> ActionButtonPtr;
+	TSharedPtr<SImage> ActionButtonImage;
+};
diff --git a/Engine/Plugins/Editor/ConsoleVariablesEditor/Source/ConsoleVariablesEditor/Private/Views/List/SConsoleVariablesEditorListValueInput.cpp b/Engine/Plugins/Editor/ConsoleVariablesEditor/Source/ConsoleVariablesEditor/Private/Views/List/SConsoleVariablesEditorListValueInput.cpp
new file mode 100644
--- /dev/null
+++ b/Engine/Plugins/Editor/ConsoleVariablesEditor/Source/ConsoleVariablesEditor/Private/Views/List/SConsoleVariablesEditorListValueInput.cpp
@@ -0,0 +1,486 @@
+// Copyright Epic Games, Inc. All Rights Reserved.
+
+#include "SConsoleVariablesEditorListValueInput.h"
+
+#include "ConsoleVariablesEditorListRow.h"
+#include "ConsoleVariablesEditorModule.h"
+#include "ConsoleVariablesEditorProjectSettings.h"
+
+#include "Widgets/Input/SNumericEntryBox.h"
+
+#define LOCTEXT_NAMESPACE "ConsoleVariablesEditor"
+
+SConsoleVariablesEditorListValueInput::~SConsoleVariablesEditorListValueInput()
+{
+	Item.Reset();
+}
+
+TSharedRef<SConsoleVariablesEditorListValueInput> SConsoleVariablesEditorListValueInput::GetInputWidget(
+	const TWeakPtr<FConsoleVariablesEditorListRow> InRow)
+{
+	const FConsoleVariablesEditorListRowPtr PinnedItem = InRow.Pin();
+
+	check(PinnedItem);
+	
+	if (PinnedItem->GetCommandInfo().IsValid())
+	{
+		const TSharedPtr<FConsoleVariablesEditorCommandInfo> PinnedInfo = PinnedItem->GetCommandInfo().Pin();
+
+		if (const TObjectPtr<IConsoleVariable> Variable = PinnedInfo->GetConsoleVariablePtr())
+		{
+			if (Variable->IsVariableFloat())
+			{
+				return SNew(SConsoleVariablesEditorListValueInput_Float, InRow);
+			}
+
+			if (Variable->IsVariableBool())
+			{
+				return SNew(SConsoleVariablesEditorListValueInput_Bool, InRow);
+			}
+
+			if (Variable->IsVariableInt())
+			{
+				return SNew(SConsoleVariablesEditorListValueInput_Int, InRow);
+			}
+
+			if (Variable->IsVariableString())
+			{
+				return SNew(SConsoleVariablesEditorListValueInput_String, InRow);
+			}
+
+			// Showflags are not considered to be any of these types, but they should be ints with a min/max of 0/2
+			if (PinnedInfo->Command.Contains("showflag", ESearchCase::IgnoreCase))
+			{
+				return SNew(SConsoleVariablesEditorListValueInput_Int, InRow, true);
+			}
+		}
+
+		// For Commands
+		return SNew(SConsoleVariablesEditorListValueInput_Command, InRow, PinnedItem->GetPresetValue());
+	}
+
+	// fallback
+	return SNew(SConsoleVariablesEditorListValueInput_String, InRow);
+}
+
+bool SConsoleVariablesEditorListValueInput::IsRowChecked() const
+{
+	return Item.Pin()->IsRowChecked();
+}
+
+void SConsoleVariablesEditorListValueInput_Float::Construct(const FArguments& InArgs,
+                                                            const TWeakPtr<FConsoleVariablesEditorListRow> InRow)
+{
+	check (InRow.IsValid());
+	
+	Item = InRow;
+
+	ProjectSettingsPtr = GetMutableDefault<UConsoleVariablesEditorProjectSettings>();
+	
+	ChildSlot
+	[
+		SAssignNew(InputWidget, SSpinBox<float>)
+		.MaxFractionalDigits(3)
+		.Value_Lambda([this]
+		{
+			check (Item.IsValid());
+			
+			if (Item.Pin()->GetWidgetCheckedState() == ECheckBoxState::Checked ||
+					(ProjectSettingsPtr &&
+						ProjectSettingsPtr->UncheckedRowDisplayType == EConsoleVariablesEditorRowDisplayType::ShowCurrentValue))
+			{
+				if (const IConsoleVariable* AsVariable = Item.Pin()->GetCommandInfo().Pin()->GetConsoleVariablePtr())
+				{
+					return FCString::Atof(*FString::SanitizeFloat(AsVariable->GetFloat()));
+				}
+			}
+
+			return FCString::Atof(*Item.Pin()->GetCachedValue());
+		})
+		.OnValueChanged_Lambda([this] (const float InValue)
+		{
+			check (Item.IsValid());
+					
+			const FString ValueAsString = FString::SanitizeFloat(InValue);
+					
+			if (const TSharedPtr<FConsoleVariablesEditorListRow> PinnedItem = Item.Pin();
+				!PinnedItem->GetCachedValue().Equals(ValueAsString))
+			{				
+				PinnedItem->GetCommandInfo().Pin()->ExecuteCommand(ValueAsString);
+
+				FConsoleVariablesEditorModule::Get().SendMultiUserConsoleVariableChange(
+					PinnedItem->GetCommandInfo().Pin()->Command, ValueAsString);
+					
+				PinnedItem->SetCachedValue(ValueAsString);
+			}
+		})
+		.IsEnabled(this, &SConsoleVariablesEditorListValueInput::IsRowChecked)
+	];
+
+	Item.Pin()->SetCachedValue(GetInputValueAsString());
+}
+
+SConsoleVariablesEditorListValueInput_Float::~SConsoleVariablesEditorListValueInput_Float()
+{
+	InputWidget.Reset();
+}
+
+void SConsoleVariablesEditorListValueInput_Float::SetInputValue(const FString& InValueAsString)
+{
+	InputWidget->SetValue(FCString::Atof(*InValueAsString));
+}
+
+FString SConsoleVariablesEditorListValueInput_Float::GetInputValueAsString()
+{
+	return FString::SanitizeFloat(GetInputValue());
+}
+
+float SConsoleVariablesEditorListValueInput_Float::GetInputValue() const
+{
+	return InputWidget->GetValue();
+}
+
+void SConsoleVariablesEditorListValueInput_Int::Construct(const FArguments& InArgs,
+                                                          const TWeakPtr<FConsoleVariablesEditorListRow> InRow,
+                                                          const bool bIsShowFlag)
+{
+	check (InRow.IsValid());
+	
+	Item = InRow;
+
+	ProjectSettingsPtr = GetMutableDefault<UConsoleVariablesEditorProjectSettings>();
+	
+	ChildSlot
+	[
+		SAssignNew(InputWidget, SSpinBox<int32>)
+		.Value_Lambda([this]
+		{
+			check (Item.IsValid());
+			
+			if (Item.Pin()->GetWidgetCheckedState() == ECheckBoxState::Checked ||
+					(ProjectSettingsPtr &&
+						ProjectSettingsPtr->UncheckedRowDisplayType == EConsoleVariablesEditorRowDisplayType::ShowCurrentValue))
+			{
+				if (const IConsoleVariable* AsVariable = Item.Pin()->GetCommandInfo().Pin()->GetConsoleVariablePtr())
+				{
+					return FCString::Atoi(*AsVariable->GetString());
+				}
+			}
+
+			return FCString::Atoi(*Item.Pin()->GetCachedValue());
+		})
+		.OnValueChanged_Lambda([this] (const int32 InValue)
+		{
+			check (Item.IsValid());
+			
+			const FString ValueAsString = FString::FromInt(InValue);
+			
+			if (const TSharedPtr<FConsoleVariablesEditorListRow> PinnedItem = Item.Pin();
+				!PinnedItem->GetCachedValue().Equals(ValueAsString))
+			{				
+				PinnedItem->GetCommandInfo().Pin()->ExecuteCommand(ValueAsString);
+
+				FConsoleVariablesEditorModule::Get().SendMultiUserConsoleVariableChange(
+					PinnedItem->GetCommandInfo().Pin()->Command, ValueAsString);
+
+				PinnedItem->SetCachedValue(ValueAsString);
+			}
+		})
+		.IsEnabled(this, &SConsoleVariablesEditorListValueInput::IsRowChecked)
+	];
+
+	if (bIsShowFlag)
+	{
+		InputWidget->SetMinSliderValue(0);
+		InputWidget->SetMaxSliderValue(2);
+
+		const int32 PresetValue = FCString::Atoi(*Item.Pin()->GetPresetValue());
+
+		Item.Pin()->SetPresetValue(FString::FromInt(FMath::Clamp(PresetValue, 0, 2)));
+	}
+	
+	Item.Pin()->SetCachedValue(GetInputValueAsString());
+}
+
+SConsoleVariablesEditorListValueInput_Int::~SConsoleVariablesEditorListValueInput_Int()
+{
+	InputWidget.Reset();
+}
+
+void SConsoleVariablesEditorListValueInput_Int::SetInputValue(const FString& InValueAsString)
+{
+	if (InValueAsString.IsNumeric())
+	{
+		InputWidget->SetValue(FCString::Atoi(*InValueAsString));
+	}
+	else
+	{
+		InputWidget->SetValue(2);
+		
+		if (InValueAsString.TrimStartAndEnd().ToLower() == "true")
+		{
+			InputWidget->SetValue(1);
+		}
+		else if (InValueAsString.TrimStartAndEnd().ToLower() == "false")
+		{
+			InputWidget->SetValue(0);
+		}
+	}
+}
+
+FString SConsoleVariablesEditorListValueInput_Int::GetInputValueAsString()
+{
+	return FString::FromInt(GetInputValue());
+}
+
+int32 SConsoleVariablesEditorListValueInput_Int::GetInputValue() const
+{
+	return InputWidget->GetValue();
+}
+
+void SConsoleVariablesEditorListValueInput_String::Construct(const FArguments& InArgs,
+                                                             const TWeakPtr<FConsoleVariablesEditorListRow> InRow)
+{
+	check (InRow.IsValid());
+	
+	Item = InRow;
+
+	ProjectSettingsPtr = GetMutableDefault<UConsoleVariablesEditorProjectSettings>();
+	
+	ChildSlot
+	[
+		SAssignNew(InputWidget, SEditableText)
+		.Text_Lambda([this]
+		{
+			check (Item.IsValid());
+			
+			if (Item.Pin()->GetWidgetCheckedState() == ECheckBoxState::Checked ||
+					(ProjectSettingsPtr &&
+						ProjectSettingsPtr->UncheckedRowDisplayType == EConsoleVariablesEditorRowDisplayType::ShowCurrentValue))
+			{
+				if (const IConsoleVariable* AsVariable = Item.Pin()->GetCommandInfo().Pin()->GetConsoleVariablePtr())
+				{
+					return FText::FromString(*AsVariable->GetString());
+				}
+			}
+
+			return FText::FromString(Item.Pin()->GetCachedValue());
+		})
+		.OnTextCommitted_Lambda([this] (const FText& InValue, ETextCommit::Type InTextCommitType)
+		{
+			check (Item.IsValid());
+			
+			const FString ValueAsString = InValue.ToString();
+			
+			if (const TSharedPtr<FConsoleVariablesEditorListRow> PinnedItem = Item.Pin();
+				!PinnedItem->GetCachedValue().Equals(ValueAsString))
+			{				
+				PinnedItem->GetCommandInfo().Pin()->ExecuteCommand(ValueAsString);
+
+				FConsoleVariablesEditorModule::Get().SendMultiUserConsoleVariableChange(
+					PinnedItem->GetCommandInfo().Pin()->Command, ValueAsString);
+
+				PinnedItem->SetCachedValue(ValueAsString);
+			}
+		})
+		.IsEnabled(this, &SConsoleVariablesEditorListValueInput::IsRowChecked)
+	];
+
+	Item.Pin()->SetCachedValue(GetInputValueAsString());
+}
+
+SConsoleVariablesEditorListValueInput_String::~SConsoleVariablesEditorListValueInput_String()
+{
+	InputWidget.Reset();
+}
+
+void SConsoleVariablesEditorListValueInput_String::SetInputValue(const FString& InValueAsString)
+{
+	InputWidget->SetText(FText::FromString(InValueAsString));
+}
+
+FString SConsoleVariablesEditorListValueInput_String::GetInputValueAsString()
+{
+	return GetInputValue();
+}
+
+FString SConsoleVariablesEditorListValueInput_String::GetInputValue() const
+{
+	return InputWidget->GetText().ToString();
+}
+
+void SConsoleVariablesEditorListValueInput_Bool::Construct(const FArguments& InArgs,
+                                                           const TWeakPtr<FConsoleVariablesEditorListRow> InRow)
+{
+	check (InRow.IsValid());
+	
+	Item = InRow;
+
+	ProjectSettingsPtr = GetMutableDefault<UConsoleVariablesEditorProjectSettings>();
+	
+	ChildSlot
+	[
+		SAssignNew(InputWidget, SButton)
+		.OnClicked_Lambda([this] ()
+		{
+			const bool bValueAsBool = GetInputValue();
+
+			SetInputValue(!bValueAsBool);
+
+			return FReply::Handled();
+		})
+		.IsEnabled(this, &SConsoleVariablesEditorListValueInput::IsRowChecked)
+		[
+			SAssignNew(ButtonText, STextBlock)
+			.Justification(ETextJustify::Center)
+			.Text_Lambda([this]()
+			{
+				check (Item.IsValid());
+			
+				if (Item.Pin()->GetWidgetCheckedState() == ECheckBoxState::Checked ||
+						(ProjectSettingsPtr &&
+							ProjectSettingsPtr->UncheckedRowDisplayType == EConsoleVariablesEditorRowDisplayType::ShowCurrentValue))
+				{
+					if (const IConsoleVariable* AsVariable = Item.Pin()->GetCommandInfo().Pin()->GetConsoleVariablePtr())
+					{
+						return FText::FromString(AsVariable->GetString());
+					}
+				}
+
+				return FText::FromString(Item.Pin()->GetCachedValue());
+			})
+		]
+	];
+	
+	Item.Pin()->SetCachedValue(GetInputValueAsString());
+}
+
+SConsoleVariablesEditorListValueInput_Bool::~SConsoleVariablesEditorListValueInput_Bool()
+{
+	InputWidget.Reset();
+	ButtonText.Reset();
+}
+
+void SConsoleVariablesEditorListValueInput_Bool::SetInputValue(const FString& InValueAsString)
+{
+	check (Item.IsValid());
+			
+	if (const TSharedPtr<FConsoleVariablesEditorListRow> PinnedItem = Item.Pin();
+		!PinnedItem->GetCachedValue().Equals(InValueAsString))
+	{				
+		PinnedItem->GetCommandInfo().Pin()->ExecuteCommand(InValueAsString);
+
+		FConsoleVariablesEditorModule::Get().SendMultiUserConsoleVariableChange(PinnedItem->GetCommandInfo().Pin()->Command, InValueAsString);
+
+		PinnedItem->SetCachedValue(InValueAsString);
+	}
+}
+
+void SConsoleVariablesEditorListValueInput_Bool::SetInputValue(const bool bNewValue)
+{	
+	SetInputValue(BoolToString(bNewValue));
+}
+
+FString SConsoleVariablesEditorListValueInput_Bool::GetInputValueAsString()
+{
+	return ButtonText->GetText().ToString();
+}
+
+bool SConsoleVariablesEditorListValueInput_Bool::GetInputValue()
+{
+	return StringToBool(GetInputValueAsString());
+}
+
+void SConsoleVariablesEditorListValueInput_Command::Construct(const FArguments& InArgs,
+	const TWeakPtr<FConsoleVariablesEditorListRow> InRow, const FString& InSavedText)
+{
+	check (InRow.IsValid());
+	
+	Item = InRow;
+
+	ProjectSettingsPtr = GetMutableDefault<UConsoleVariablesEditorProjectSettings>();
+	
+	ChildSlot
+	[
+		SNew(SHorizontalBox)
+
+		+SHorizontalBox::Slot()
+		.VAlign(VAlign_Center)
+		[
+			SAssignNew(InputText, SEditableText)
+			.Text(FText::FromString(InSavedText))
+			.HintText(LOCTEXT("CommandValueTypeRowInputHintText", "Value..."))
+			.IsEnabled(this, &SConsoleVariablesEditorListValueInput::IsRowChecked)
+			.Visibility(Item.Pin()->GetCommandInfo().Pin()->ObjectType ==
+				FConsoleVariablesEditorCommandInfo::EConsoleObjectType::Command ?
+				EVisibility::Visible : EVisibility::Collapsed)
+		]
+
+		+SHorizontalBox::Slot()
+		.VAlign(VAlign_Center)
+		[
+			SAssignNew(InputWidget, SButton)
+			.OnClicked_Lambda([this] ()
+			{
+				check (Item.IsValid());
+				
+				const FString ValueAsString = InputText->GetText().ToString();
+				
+				if (const TSharedPtr<FConsoleVariablesEditorListRow> PinnedItem = Item.Pin())
+				{
+					// Only execute if this row is a non-value command OR
+					// it's a value-based command whose cached value differs from the input value
+					if (PinnedItem->GetCommandInfo().Pin()->ObjectType ==
+						FConsoleVariablesEditorCommandInfo::EConsoleObjectType::NullObject ||
+						(PinnedItem->GetCommandInfo().Pin()->ObjectType ==
+						FConsoleVariablesEditorCommandInfo::EConsoleObjectType::Command &&
+						!PinnedItem->GetCachedValue().Equals(ValueAsString)))
+					{
+						PinnedItem->GetCommandInfo().Pin()->ExecuteCommand(ValueAsString);
+
+						FConsoleVariablesEditorModule::Get().SendMultiUserConsoleVariableChange(
+							PinnedItem->GetCommandInfo().Pin()->Command, ValueAsString);
+
+						PinnedItem->SetCachedValue(ValueAsString);
+					}
+
+					return FReply::Handled();
+				}
+
+				return FReply::Unhandled();
+			})
+			.IsEnabled(this, &SConsoleVariablesEditorListValueInput::IsRowChecked)
+			.ContentPadding(FMargin(0.f))
+			[
+				SNew(STextBlock)
+				.Justification(ETextJustify::Center)
+				.Text(LOCTEXT("ConsoleCommandExecutionButtonText","Execute"))
+			]
+		]
+	];
+
+	Item.Pin()->SetCachedValue(GetInputValueAsString());
+}
+
+SConsoleVariablesEditorListValueInput_Command::~SConsoleVariablesEditorListValueInput_Command()
+{
+	InputWidget.Reset();
+	InputText.Reset();
+}
+
+void SConsoleVariablesEditorListValueInput_Command::SetInputValue(const FString& InValueAsString)
+{
+	InputText->SetText(FText::FromString(InValueAsString));
+}
+
+FString SConsoleVariablesEditorListValueInput_Command::GetInputValueAsString()
+{
+	return GetInputValue();
+}
+
+FString SConsoleVariablesEditorListValueInput_Command::GetInputValue() const
+{
+	return InputText->GetText().ToString();
+}
+
+#undef LOCTEXT_NAMESPACE
diff --git a/Engine/Plugins/Editor/ConsoleVariablesEditor/Source/ConsoleVariablesEditor/Private/Views/List/SConsoleVariablesEditorListValueInput.h b/Engine/Plugins/Editor/ConsoleVariablesEditor/Source/ConsoleVariablesEditor/Private/Views/List/SConsoleVariablesEditorListValueInput.h
new file mode 100644
--- /dev/null
+++ b/Engine/Plugins/Editor/ConsoleVariablesEditor/Source/ConsoleVariablesEditor/Private/Views/List/SConsoleVariablesEditorListValueInput.h
@@ -0,0 +1,161 @@
+// Copyright Epic Games, Inc. All Rights Reserved.
+
+#pragma once
+
+#include "Widgets/Input/SButton.h"
+#include "Widgets/Input/SSpinBox.h"
+#include "Widgets/SCompoundWidget.h"
+#include "Widgets/Text/STextBlock.h"
+
+struct FConsoleVariablesEditorListRow;
+
+class UConsoleVariablesEditorProjectSettings;
+
+class SConsoleVariablesEditorListValueInput : public SCompoundWidget
+{
+public:
+	virtual ~SConsoleVariablesEditorListValueInput() override;
+
+	static TSharedRef<SConsoleVariablesEditorListValueInput> GetInputWidget(const TWeakPtr<FConsoleVariablesEditorListRow> InRow);
+	
+	virtual void SetInputValue(const FString& InValueAsString) = 0;
+	virtual FString GetInputValueAsString() = 0;
+
+	bool IsRowChecked() const;
+
+protected:
+	
+	TWeakPtr<FConsoleVariablesEditorListRow> Item = nullptr;
+
+	TObjectPtr<UConsoleVariablesEditorProjectSettings> ProjectSettingsPtr = nullptr;
+};
+
+class SConsoleVariablesEditorListValueInput_Float : public SConsoleVariablesEditorListValueInput
+{
+public:
+	
+	SLATE_BEGIN_ARGS(SConsoleVariablesEditorListValueInput_Float)
+	{}
+
+	SLATE_END_ARGS()
+
+	void Construct(const FArguments& InArgs, const TWeakPtr<FConsoleVariablesEditorListRow> InRow);
+	
+	virtual ~SConsoleVariablesEditorListValueInput_Float() override;
+	
+	virtual void SetInputValue(const FString& InValueAsString) override;
+	virtual FString GetInputValueAsString() override;
+
+	float GetInputValue() const;
+
+private:
+	
+	TSharedPtr<SSpinBox<float>> InputWidget;
+};
+
+class SConsoleVariablesEditorListValueInput_Int : public SConsoleVariablesEditorListValueInput
+{
+public:
+	
+	SLATE_BEGIN_ARGS(SConsoleVariablesEditorListValueInput_Int)
+	{}
+
+	SLATE_END_ARGS()
+
+	void Construct(
+		const FArguments& InArgs, const TWeakPtr<FConsoleVariablesEditorListRow> InRow, const bool bIsShowFlag = false);
+
+	virtual ~SConsoleVariablesEditorListValueInput_Int() override;
+	
+	virtual void SetInputValue(const FString& InValueAsString) override;
+	virtual FString GetInputValueAsString() override;
+
+	int32 GetInputValue() const;
+
+private:
+	
+	TSharedPtr<SSpinBox<int32>> InputWidget;
+};
+
+class SConsoleVariablesEditorListValueInput_String : public SConsoleVariablesEditorListValueInput
+{
+public:
+	
+	SLATE_BEGIN_ARGS(SConsoleVariablesEditorListValueInput_String)
+	{}
+
+	SLATE_END_ARGS()
+
+	void Construct(const FArguments& InArgs, const TWeakPtr<FConsoleVariablesEditorListRow> InRow);
+
+	virtual ~SConsoleVariablesEditorListValueInput_String() override;
+	
+	virtual void SetInputValue(const FString& InValueAsString) override;
+	virtual FString GetInputValueAsString() override;
+	
+	FString GetInputValue() const;
+
+private:
+	
+	TSharedPtr<SEditableText> InputWidget;
+};
+
+class SConsoleVariablesEditorListValueInput_Bool : public SConsoleVariablesEditorListValueInput
+{
+public:
+	
+	SLATE_BEGIN_ARGS(SConsoleVariablesEditorListValueInput_Bool)
+	{}
+
+	SLATE_END_ARGS()
+
+	void Construct(const FArguments& InArgs, const TWeakPtr<FConsoleVariablesEditorListRow> InRow);
+
+	virtual ~SConsoleVariablesEditorListValueInput_Bool() override;
+	
+	virtual void SetInputValue(const FString& InValueAsString) override;
+	void SetInputValue(const bool bNewValue);
+	virtual FString GetInputValueAsString() override;
+
+	bool GetInputValue();
+
+	static bool StringToBool(const FString& InString)
+	{
+		return InString.TrimStartAndEnd().ToLower() == "true";
+	}
+
+	static FString BoolToString(const bool bNewBool)
+	{
+		return bNewBool ? "true" : "false";
+	}
+
+private:
+	
+	TSharedPtr<SButton> InputWidget;
+	TSharedPtr<STextBlock> ButtonText;
+};
+
+class SConsoleVariablesEditorListValueInput_Command : public SConsoleVariablesEditorListValueInput
+{
+public:
+	
+	SLATE_BEGIN_ARGS(SConsoleVariablesEditorListValueInput_Command)
+	{}
+
+	SLATE_END_ARGS()
+
+	void Construct(
+		const FArguments& InArgs, const TWeakPtr<FConsoleVariablesEditorListRow> InRow, const FString& InSavedText);
+	
+	virtual ~SConsoleVariablesEditorListValueInput_Command() override;
+	
+	virtual void SetInputValue(const FString& InValueAsString) override;
+	virtual FString GetInputValueAsString() override;
+
+	FString GetInputValue() const;
+
+private:
+	
+	TSharedPtr<SButton> InputWidget;
+	TSharedPtr<SEditableText> InputText;
+};
diff --git a/Engine/Plugins/Editor/ConsoleVariablesEditor/Source/ConsoleVariablesEditor/Private/Views/MainPanel/ConsoleVariablesEditorMainPanel.cpp b/Engine/Plugins/Editor/ConsoleVariablesEditor/Source/ConsoleVariablesEditor/Private/Views/MainPanel/ConsoleVariablesEditorMainPanel.cpp
new file mode 100644
--- /dev/null
+++ b/Engine/Plugins/Editor/ConsoleVariablesEditor/Source/ConsoleVariablesEditor/Private/Views/MainPanel/ConsoleVariablesEditorMainPanel.cpp
@@ -0,0 +1,298 @@
+// Copyright Epic Games, Inc. All Rights Reserved.
+
+#include "Views/MainPanel/ConsoleVariablesEditorMainPanel.h"
+
+#include "ConcertMessages.h"
+#include "ConsoleVariablesAsset.h"
+#include "ConsoleVariablesEditorCommandInfo.h"
+#include "ConsoleVariablesEditorLog.h"
+#include "ConsoleVariablesEditorModule.h"
+#include "Views/List/ConsoleVariablesEditorList.h"
+#include "Views/MainPanel/SConsoleVariablesEditorMainPanel.h"
+
+#include "FileHelpers.h"
+#include "Framework/Application/SlateApplication.h"
+#include "MultiUser/ConsoleVariableSyncData.h"
+
+
+FConsoleVariablesEditorMainPanel::FConsoleVariablesEditorMainPanel()
+{
+	EditorList = MakeShared<FConsoleVariablesEditorList>();
+	
+	OnConnectionChangedHandle = MultiUserManager.OnConnectionChange().AddStatic(
+		&FConsoleVariablesEditorMainPanel::OnConnectionChanged);
+	OnRemoteCVarChangeHandle = MultiUserManager.OnRemoteCVarChange().AddStatic(
+		&FConsoleVariablesEditorMainPanel::OnRemoteCvarChange);
+}
+
+FConsoleVariablesEditorMainPanel::~FConsoleVariablesEditorMainPanel()
+{
+	MainPanelWidget.Reset();
+	EditorList.Reset();
+	MultiUserManager.OnConnectionChange().Remove(OnConnectionChangedHandle);
+	MultiUserManager.OnRemoteCVarChange().Remove(OnRemoteCVarChangeHandle);
+}
+
+TSharedRef<SWidget> FConsoleVariablesEditorMainPanel::GetOrCreateWidget()
+{
+	if (!MainPanelWidget.IsValid())
+	{
+		SAssignNew(MainPanelWidget, SConsoleVariablesEditorMainPanel, SharedThis(this));
+	}
+
+	return MainPanelWidget.ToSharedRef();
+}
+
+FConsoleVariablesEditorModule& FConsoleVariablesEditorMainPanel::GetConsoleVariablesModule()
+{
+	return FConsoleVariablesEditorModule::Get();
+}
+
+TObjectPtr<UConsoleVariablesAsset> FConsoleVariablesEditorMainPanel::GetEditingAsset()
+{
+	return GetConsoleVariablesModule().GetPresetAsset();
+}
+
+void FConsoleVariablesEditorMainPanel::AddConsoleObjectToCurrentPreset(
+	const FString InConsoleCommand, const FString InValue, const bool bScrollToNewRow) const
+{
+	if (const TObjectPtr<UConsoleVariablesAsset> Asset = GetEditingAsset())
+	{
+		Asset->AddOrSetConsoleObjectSavedData(
+			{
+				InConsoleCommand,
+				InValue,
+				ECheckBoxState::Checked
+			}
+		);
+
+		if (GetEditorList().Pin()->GetListMode() == FConsoleVariablesEditorList::EConsoleVariablesEditorListMode::Preset)
+		{
+			RebuildList(bScrollToNewRow ? InConsoleCommand : "");
+		}
+	}
+}
+
+FReply FConsoleVariablesEditorMainPanel::ValidateConsoleInputAndAddToCurrentPreset(const FText& CommittedText) const
+{
+	// Clear Search
+	if (const TSharedPtr<FConsoleVariablesEditorList> PinnedList = GetEditorList().Pin())
+	{
+		PinnedList->SetSearchString("");
+	}
+	
+	const FString CommandString = CommittedText.ToString().TrimStartAndEnd();
+	FString CommandKey;
+	FString ValueString;
+	
+	if (CommandString.Contains(" "))
+	{
+		CommandString.Split(TEXT(" "), &CommandKey, &ValueString);
+	}
+	else
+	{
+		CommandKey = CommandString;
+	}
+
+	if (IConsoleObject* ConsoleObject = IConsoleManager::Get().FindConsoleObject(*CommandKey)) 
+	{
+		if (ValueString.IsEmpty())
+		{
+			if (const IConsoleVariable* AsVariable = ConsoleObject->AsVariable())
+			{
+				ValueString = AsVariable->GetString();
+			}
+		}
+		
+		AddConsoleObjectToCurrentPreset(
+			CommandKey,
+			ValueString,
+			true
+		);
+	}
+	else if (CommandString.IsEmpty())
+	{
+		UE_LOG(LogConsoleVariablesEditor, Warning, TEXT("%hs: Input is blank."), __FUNCTION__);
+
+		return FReply::Unhandled();
+	}
+	// Try to execute the whole command. Some commands are not registered, but are parsed externally so they won't be found by IConsoleManager::FindConsoleObject
+	else if (GEngine->Exec(FConsoleVariablesEditorCommandInfo::GetCurrentWorld(), *CommandString))
+	{
+		AddConsoleObjectToCurrentPreset(
+			CommandString,
+			"",
+			true
+			);
+	}
+	else
+	{
+		UE_LOG(LogConsoleVariablesEditor, Warning,
+			TEXT("%hs: Input '%s' is not a recognized console variable or command."),
+			__FUNCTION__, *CommandString);
+
+		return FReply::Unhandled();
+	}
+
+	return FReply::Handled();
+}
+
+void FConsoleVariablesEditorMainPanel::RebuildList(const FString InConsoleCommandToScrollTo, bool bShouldCacheValues) const
+{
+	if (EditorList.IsValid())
+	{
+		EditorList->RebuildList(InConsoleCommandToScrollTo);
+	}
+}
+
+void FConsoleVariablesEditorMainPanel::RefreshList() const
+{
+	if (EditorList.IsValid())
+	{
+		EditorList->RefreshList();
+	}
+}
+
+void FConsoleVariablesEditorMainPanel::UpdatePresetValuesForSave(TObjectPtr<UConsoleVariablesAsset> InAsset) const
+{
+	if (EditorList.IsValid())
+	{
+		EditorList->UpdatePresetValuesForSave(InAsset);
+	}
+}
+
+void FConsoleVariablesEditorMainPanel::RefreshMultiUserDetails() const
+{
+	if (MainPanelWidget.IsValid())
+	{
+		MainPanelWidget->RefreshMultiUserDetails();
+	}
+}
+
+void FConsoleVariablesEditorMainPanel::SaveCurrentPreset()
+{
+	if (const TWeakObjectPtr<UConsoleVariablesAsset> EditingAsset = GetEditingAsset();
+		EditingAsset.IsValid() && ReferenceAssetOnDisk.IsValid())
+	{
+		UpdatePresetValuesForSave(EditingAsset.Get());
+		
+		if (UPackage* ReferencePackage = ReferenceAssetOnDisk.Get()->GetPackage())
+		{
+			ReferenceAssetOnDisk->CopyFrom(EditingAsset.Get());
+
+			UEditorLoadingAndSavingUtils::SavePackages({ ReferencePackage }, false);
+
+			return;
+		}
+	}
+
+	// Fallback
+	SaveCurrentPresetAs();
+}
+
+void FConsoleVariablesEditorMainPanel::SaveSpecificPreset(const TObjectPtr<UConsoleVariablesAsset> InPreset) const
+{
+	if (InPreset)
+	{		
+		if (UPackage* ReferencePackage = InPreset->GetPackage())
+		{
+			UEditorLoadingAndSavingUtils::SavePackages({ ReferencePackage }, false);
+
+			return;
+		}
+	}
+
+	// Fallback
+	SaveSpecificPresetAs(InPreset);
+}
+
+void FConsoleVariablesEditorMainPanel::SaveCurrentPresetAs()
+{
+	if (const TWeakObjectPtr<UConsoleVariablesAsset> EditingAsset = GetEditingAsset(); EditingAsset.IsValid())
+	{
+		UpdatePresetValuesForSave(EditingAsset.Get());
+			
+		TArray<UObject*> SavedAssets;
+		FEditorFileUtils::SaveAssetsAs({ EditingAsset.Get() }, SavedAssets);
+
+		if (SavedAssets.Num())
+		{
+			UConsoleVariablesAsset* SavedAsset = Cast<UConsoleVariablesAsset>(SavedAssets[0]);
+
+			if (ensure(SavedAsset))
+			{
+				ReferenceAssetOnDisk = SavedAsset;
+			}
+		}
+	}
+}
+
+void FConsoleVariablesEditorMainPanel::SaveSpecificPresetAs(const TObjectPtr<UConsoleVariablesAsset> InPreset) const
+{
+	if (InPreset)
+	{
+		TArray<UObject*> SavedAssets;
+		FEditorFileUtils::SaveAssetsAs({ InPreset }, SavedAssets);
+	}
+}
+
+void FConsoleVariablesEditorMainPanel::ImportPreset(const FAssetData& InPresetAsset)
+{
+	FSlateApplication::Get().DismissAllMenus();
+	
+	if (UConsoleVariablesAsset* Preset = CastChecked<UConsoleVariablesAsset>(InPresetAsset.GetAsset()))
+	{
+		if (const TObjectPtr<UConsoleVariablesAsset> EditingAsset = GetEditingAsset(); ImportPreset_Impl(Preset, EditingAsset))
+		{
+			EditorList->RebuildList("", false);
+		}
+	}
+}
+
+void FConsoleVariablesEditorMainPanel::ImportPreset(const TObjectPtr<UConsoleVariablesAsset> InPreset)
+{
+	FSlateApplication::Get().DismissAllMenus();
+
+	if (const TObjectPtr<UConsoleVariablesAsset> EditingAsset = GetEditingAsset(); ImportPreset_Impl(InPreset, EditingAsset))
+	{
+		EditorList->RebuildList("", false);
+	}
+}
+
+bool FConsoleVariablesEditorMainPanel::ImportPreset_Impl(
+	const TObjectPtr<UConsoleVariablesAsset> Preset, const TObjectPtr<UConsoleVariablesAsset> EditingAsset)
+{
+	if (Preset && EditingAsset)
+	{
+		ReferenceAssetOnDisk = Preset;
+
+		EditingAsset->Modify();
+		EditingAsset->CopyFrom(Preset);
+
+		return true;
+	}
+
+	return false;
+}
+
+void FConsoleVariablesEditorMainPanel::OnConnectionChanged(EConcertConnectionStatus Status)
+{
+	switch (Status)
+	{
+		case EConcertConnectionStatus::Connected:
+			UE_LOG(LogConsoleVariablesEditor, Display, TEXT("Multi-user has connected to a session."));
+			break;
+		case EConcertConnectionStatus::Disconnected:
+			UE_LOG(LogConsoleVariablesEditor, Display, TEXT("Multi-user has disconnected from session."));
+			break;
+		default:
+			break;
+	}
+}
+
+void FConsoleVariablesEditorMainPanel::OnRemoteCvarChange(const FString InName, const FString InValue)
+{
+	FConsoleVariablesEditorModule& ConsoleVariablesEditorModule = FConsoleVariablesEditorModule::Get();
+
+	ConsoleVariablesEditorModule.OnRemoteCvarChanged(InName, InValue);
+}
diff --git a/Engine/Plugins/Editor/ConsoleVariablesEditor/Source/ConsoleVariablesEditor/Private/Views/MainPanel/ConsoleVariablesEditorMainPanel.h b/Engine/Plugins/Editor/ConsoleVariablesEditor/Source/ConsoleVariablesEditor/Private/Views/MainPanel/ConsoleVariablesEditorMainPanel.h
new file mode 100644
--- /dev/null
+++ b/Engine/Plugins/Editor/ConsoleVariablesEditor/Source/ConsoleVariablesEditor/Private/Views/MainPanel/ConsoleVariablesEditorMainPanel.h
@@ -0,0 +1,93 @@
+// Copyright Epic Games, Inc. All Rights Reserved.
+
+#pragma once
+
+#include "ConsoleVariablesEditorModule.h"
+#include "Views/List/ConsoleVariablesEditorList.h"
+#include "Widgets/SWidget.h"
+
+class FConsoleVariablesEditorList;
+class SConsoleVariablesEditorMainPanel;
+class UConsoleVariablesAsset;
+
+class FConsoleVariablesEditorMainPanel : public TSharedFromThis<FConsoleVariablesEditorMainPanel>
+{
+public:
+	FConsoleVariablesEditorMainPanel();
+
+	~FConsoleVariablesEditorMainPanel();
+
+	TSharedRef<SWidget> GetOrCreateWidget();
+
+	static FConsoleVariablesEditorModule& GetConsoleVariablesModule();
+	static TObjectPtr<UConsoleVariablesAsset> GetEditingAsset();
+
+	void AddConsoleObjectToCurrentPreset(
+		const FString InConsoleCommand, const FString InValue,
+		const bool bScrollToNewRow = false) const;
+
+	FReply ValidateConsoleInputAndAddToCurrentPreset(const FText& CommittedText) const;
+
+	/*
+	 * Regenerate the list items and refresh the list. Call when adding or removing variables.
+	 * @param bShouldCacheValues If true, the current list's current values will be cached and then restored when the list is rebuilt. Otherwise preset values will be used.
+	 */
+	void RebuildList(const FString InConsoleCommandToScrollTo = "", bool bShouldCacheValues = true) const;
+
+	/**
+	 * Refresh filters and sorting.
+	 * Useful for when the list state has gone stale but the variable count has not changed.
+	 */
+	void RefreshList() const;
+	
+	void UpdatePresetValuesForSave(TObjectPtr<UConsoleVariablesAsset> InAsset) const;
+
+	void RefreshMultiUserDetails() const;
+
+	// Save / Load
+
+	void SaveCurrentPreset();
+	void SaveSpecificPreset(const TObjectPtr<UConsoleVariablesAsset> InPreset) const;
+	void SaveCurrentPresetAs();
+	void SaveSpecificPresetAs(const TObjectPtr<UConsoleVariablesAsset> InPreset) const;
+	void ImportPreset(const FAssetData& InPresetAsset);
+	void ImportPreset(const TObjectPtr<UConsoleVariablesAsset> InPreset);
+
+	TWeakObjectPtr<UConsoleVariablesAsset> GetReferenceAssetOnDisk() const
+	{
+		return ReferenceAssetOnDisk;
+	}
+
+	TWeakPtr<FConsoleVariablesEditorList> GetEditorList() const
+	{
+		return EditorList;
+	}
+
+	FConsoleVariablesEditorList::EConsoleVariablesEditorListMode GetEditorListMode() const
+	{
+		return EditorList->GetListMode();
+	}
+
+	UE::ConsoleVariables::MultiUser::Private::FManager& GetMultiUserManager()
+	{
+		return MultiUserManager;
+	}
+
+private:
+
+	bool ImportPreset_Impl(const TObjectPtr<UConsoleVariablesAsset> Preset, const TObjectPtr<UConsoleVariablesAsset> EditingAsset);
+
+	TSharedPtr<SConsoleVariablesEditorMainPanel> MainPanelWidget;
+
+	// The non-transient loaded asset from which we will copy to the transient asset for editing
+	TWeakObjectPtr<UConsoleVariablesAsset> ReferenceAssetOnDisk;
+
+	TSharedPtr<FConsoleVariablesEditorList> EditorList;
+
+	static void OnConnectionChanged(EConcertConnectionStatus Status);
+	static void OnRemoteCvarChange(const FString InName, const FString InValue);
+
+	UE::ConsoleVariables::MultiUser::Private::FManager MultiUserManager;
+	FDelegateHandle OnConnectionChangedHandle;
+	FDelegateHandle OnRemoteCVarChangeHandle;
+};
diff --git a/Engine/Plugins/Editor/ConsoleVariablesEditor/Source/ConsoleVariablesEditor/Private/Views/MainPanel/SConsoleVariablesEditorCustomConsoleInputBox.cpp b/Engine/Plugins/Editor/ConsoleVariablesEditor/Source/ConsoleVariablesEditor/Private/Views/MainPanel/SConsoleVariablesEditorCustomConsoleInputBox.cpp
new file mode 100644
--- /dev/null
+++ b/Engine/Plugins/Editor/ConsoleVariablesEditor/Source/ConsoleVariablesEditor/Private/Views/MainPanel/SConsoleVariablesEditorCustomConsoleInputBox.cpp
@@ -0,0 +1,380 @@
+// Copyright Epic Games, Inc. All Rights Reserved.
+
+#include "Views/MainPanel/SConsoleVariablesEditorCustomConsoleInputBox.h"
+
+#include "ConsoleVariablesEditorLog.h"
+#include "ConsoleVariablesEditorModule.h"
+#include "Views/List/ConsoleVariablesEditorList.h"
+#include "Views/MainPanel/SConsoleVariablesEditorMainPanel.h"
+
+#include "Framework/Application/SlateApplication.h"
+#include "OutputLog/Public/OutputLogModule.h"
+#include "Styling/AppStyle.h"
+#include "Styling/StyleColors.h"
+#include "Views/Widgets/SConsoleVariablesEditorTooltipWidget.h"
+#include "Widgets/Input/SEditableTextBox.h"
+#include "Widgets/Layout/SBorder.h"
+#include "Widgets/Layout/SBox.h"
+#include "Widgets/Text/STextBlock.h"
+
+#define LOCTEXT_NAMESPACE "ConsoleVariablesEditor"
+
+void SConsoleVariablesEditorCustomConsoleInputBox::Construct(
+	const FArguments& InArgs, TWeakPtr<SConsoleVariablesEditorMainPanel> InMainPanelWidget)
+{
+	check(InMainPanelWidget.IsValid());
+
+	MainPanelWidget = InMainPanelWidget;
+	
+	ChildSlot
+	[
+		SAssignNew( SuggestionBox, SMenuAnchor )
+		.Method(GIsEditor ? EPopupMethod::CreateNewWindow : EPopupMethod::UseCurrentWindow)
+		.Placement( MenuPlacement_BelowAnchor )
+		[
+			SAssignNew(InputText, SEditableTextBox)
+			.Font(FEditorStyle::Get().GetWidgetStyle<FTextBlockStyle>("Log.Normal").Font)
+			.HintText(LOCTEXT("ConsoleCommandExecutorHintText", "Enter Console Command"))
+			.OnTextChanged(this, &SConsoleVariablesEditorCustomConsoleInputBox::OnTextChanged)
+			.OnKeyCharHandler(this, &SConsoleVariablesEditorCustomConsoleInputBox::OnKeyCharHandler)
+			.OnKeyDownHandler(this, &SConsoleVariablesEditorCustomConsoleInputBox::OnKeyDownHandler)
+			.OnTextCommitted_Lambda([this] (const FText& InText, const ETextCommit::Type CommitType)
+			{
+				// Hide widget when focus is lost
+				if (CommitType == ETextCommit::OnUserMovedFocus)
+				{
+					SuggestionBox->SetIsOpen(false);
+					SetVisibility(EVisibility::Collapsed);
+				}
+			})
+			.ClearKeyboardFocusOnCommit(false)
+		
+		]
+		.MenuContent
+		(
+			SNew(SBorder)
+			.BorderImage(FEditorStyle::GetBrush("Menu.Background"))
+			.Padding( FMargin(2) )
+			[
+				SNew(SBox)
+				.HeightOverride(250) // avoids flickering, ideally this would be adaptive to the content without flickering
+				.MinDesiredWidth(300)
+				.MaxDesiredWidth_Lambda([this] ()
+				{
+					// Limit the width of the suggestions list to the work area that this widget currently resides on
+					const FSlateRect WidgetRect(
+						GetCachedGeometry().GetAbsolutePosition(),
+						GetCachedGeometry().GetAbsolutePosition() + GetCachedGeometry().GetAbsoluteSize());
+					const FSlateRect WidgetWorkArea = FSlateApplication::Get().GetWorkArea(WidgetRect);
+					return FMath::Max(300.0f, WidgetWorkArea.GetSize().X - 12.0f);
+				})
+				[
+					SAssignNew(SuggestionListView, SListView< TSharedPtr<FString> >)
+					.ListItemsSource(&Suggestions.SuggestionsList)
+					.SelectionMode( ESelectionMode::Single )							// Ideally the mouse over would not highlight while keyboard controls the UI
+					.OnGenerateRow_Lambda([this] (TSharedPtr<FString> Text, const TSharedRef<STableViewBase>& OwnerTable)
+					{
+						check(Text.IsValid());
+
+						FString SanitizedText = *Text;
+						SanitizedText.ReplaceInline(TEXT("\r\n"), TEXT("\n"), ESearchCase::CaseSensitive);
+						SanitizedText.ReplaceInline(TEXT("\r"), TEXT(" "), ESearchCase::CaseSensitive);
+						SanitizedText.ReplaceInline(TEXT("\n"), TEXT(" "), ESearchCase::CaseSensitive);
+
+						const FText DisplayText = FText::FromString(SanitizedText);
+						
+						FConsoleVariablesEditorModule& ConsoleVariablesEditorModule = FConsoleVariablesEditorModule::Get();
+
+						TSharedPtr<IToolTip> ToolTip;
+
+						if (const TWeakPtr<FConsoleVariablesEditorCommandInfo> MatchingCommand =
+							ConsoleVariablesEditorModule.FindCommandInfoByName(SanitizedText); MatchingCommand.IsValid())
+						{
+							ToolTip = SConsoleVariablesEditorTooltipWidget::MakeTooltip(
+								SanitizedText,
+								MatchingCommand.Pin()->GetHelpText());
+						}
+
+						return
+							SNew(STableRow< TSharedPtr<FString> >, OwnerTable)
+							[
+								SNew(STextBlock)
+								.Text(DisplayText)
+								.TextStyle(FEditorStyle::Get(), "Log.Normal")
+								.HighlightText(Suggestions.SuggestionsHighlight)
+								.ColorAndOpacity(FSlateColor::UseForeground())
+								.ToolTip(ToolTip.IsValid() ? ToolTip : nullptr)
+							];
+					})
+					.OnSelectionChanged_Lambda([this] (TSharedPtr<FString> NewValue, ESelectInfo::Type SelectInfo)
+					{
+						if(bIgnoreUIUpdate)
+						{
+							return;
+						}
+						
+						Suggestions.SelectedSuggestion = Suggestions.SuggestionsList.IndexOfByPredicate([&NewValue](const TSharedPtr<FString>& InSuggestion)
+						{
+							return InSuggestion == NewValue;
+						});
+
+						MarkActiveSuggestion();
+
+						// If the user selected this suggestion by clicking on it, then go ahead and close the suggestion
+						// box as they've chosen the suggestion they're interested in.
+						if( SelectInfo == ESelectInfo::OnMouseClick )
+						{
+							SuggestionBox->SetIsOpen( false );
+
+							// Jump the caret to the end of the newly auto-completed line. This makes it so that selecting
+							// an option doesn't leave the cursor in the middle of the suggestion (which makes it hard to 
+							// ctrl-back out, or to type "?" for help, etc.)
+							InputText->GoTo(ETextLocation::EndOfDocument);
+						}
+					})
+					.ItemHeight(18)
+				]
+			]
+		)
+	];
+}
+
+SConsoleVariablesEditorCustomConsoleInputBox::~SConsoleVariablesEditorCustomConsoleInputBox()
+{
+	MainPanelWidget.Reset();
+	InputText.Reset();
+	SuggestionBox.Reset();
+	SuggestionListView.Reset();
+}
+
+FReply SConsoleVariablesEditorCustomConsoleInputBox::OnPreviewKeyDown(const FGeometry& MyGeometry, const FKeyEvent& KeyEvent)
+{
+	if(SuggestionBox->IsOpen())
+	{
+		if(KeyEvent.GetKey() == EKeys::Up || KeyEvent.GetKey() == EKeys::Down)
+		{
+			Suggestions.StepSelectedSuggestion(KeyEvent.GetKey() == EKeys::Up ? -1 : +1);
+			MarkActiveSuggestion();
+
+			return FReply::Handled();
+		}
+		else if (KeyEvent.GetKey() == EKeys::Tab)
+		{
+			if (Suggestions.HasSuggestions())
+			{
+				if (Suggestions.HasSelectedSuggestion())
+				{
+					Suggestions.StepSelectedSuggestion(KeyEvent.IsShiftDown() ? -1 : +1);
+				}
+				else
+				{
+					Suggestions.SelectedSuggestion = 0;
+				}
+				MarkActiveSuggestion();
+			}
+			
+			return FReply::Handled();
+		}
+		else if (KeyEvent.GetKey() == EKeys::Escape)
+		{
+			SuggestionBox->SetIsOpen(false);
+			return FReply::Handled();
+		}
+	}
+	else
+	{
+		if(KeyEvent.GetKey() == EKeys::Up)
+		{
+			// If the command field isn't empty we need you to have pressed Control+Up to summon the history (to make sure you're not just using caret navigation)
+			const bool bShowHistory = InputText->GetText().IsEmpty() || KeyEvent.IsControlDown();
+			if (bShowHistory)
+			{
+				TArray<FString> History;
+				IConsoleManager::Get().GetConsoleHistory(TEXT(""), History);
+				
+				SetSuggestions(History, FText::GetEmpty());
+				
+				if(Suggestions.HasSuggestions())
+				{
+					Suggestions.StepSelectedSuggestion(-1);
+					MarkActiveSuggestion();
+				}
+			}
+
+			// Need to always handle this for single-line controls to avoid them invoking widget navigation
+			if (bShowHistory)
+			{
+				return FReply::Handled();
+			}
+		}
+	}
+
+	return FReply::Unhandled();
+}
+
+bool SConsoleVariablesEditorCustomConsoleInputBox::TakeKeyboardFocus() const
+{
+	check(InputText.IsValid());
+	
+	return FSlateApplication::Get().SetKeyboardFocus(InputText.ToSharedRef());
+}
+
+void SConsoleVariablesEditorCustomConsoleInputBox::OnTextChanged(const FText& InText)
+{
+	if(bIgnoreUIUpdate)
+	{
+		return;
+	}
+	
+	const FString& InputTextStr = InputText->GetText().ToString();
+	if(!InputTextStr.IsEmpty())
+	{
+		TArray<FString> AutoCompleteList;
+		
+		auto OnConsoleVariable = [&AutoCompleteList](const TCHAR *Name, IConsoleObject* CVar)
+		{
+#if (UE_BUILD_SHIPPING || UE_BUILD_TEST)
+			if (CVar->TestFlags(ECVF_Cheat))
+			{
+				return;
+			}
+#endif // (UE_BUILD_SHIPPING || UE_BUILD_TEST)
+			if (CVar->TestFlags(ECVF_Unregistered))
+			{
+				return;
+			}
+
+			AutoCompleteList.Add(Name);
+		};
+
+		IConsoleManager::Get().ForEachConsoleObjectThatContains(FConsoleObjectVisitor::CreateLambda(OnConsoleVariable), *InputTextStr);
+
+		AutoCompleteList.Sort([InputTextStr](const FString& A, const FString& B)
+		{ 
+			if (A.StartsWith(InputTextStr))
+			{
+				if (!B.StartsWith(InputTextStr))
+				{
+					return true;
+				}
+			}
+			else
+			{
+				if (B.StartsWith(InputTextStr))
+				{
+					return false;
+				}
+			}
+
+			return A < B;
+
+		});
+
+
+		SetSuggestions(AutoCompleteList, FText::FromString(InputTextStr));
+	}
+	else
+	{
+		ClearSuggestions();
+	}
+}
+
+FReply SConsoleVariablesEditorCustomConsoleInputBox::OnKeyCharHandler(const FGeometry& MyGeometry, const FCharacterEvent& InCharacterEvent) const
+{
+	// Intercept tab if used for auto-complete
+	if (InCharacterEvent.GetCharacter() == '\t')
+	{
+		return FReply::Handled();
+	}
+
+	if (InCharacterEvent.GetModifierKeys().AnyModifiersDown() && InCharacterEvent.GetCharacter() == ' ')
+	{	
+		// Ignore space bar + a modifier key.  It should not type a space as this is used by other keyboard shortcuts
+		return FReply::Handled();
+	}
+
+	return FReply::Unhandled();
+}
+
+FReply SConsoleVariablesEditorCustomConsoleInputBox::OnKeyDownHandler(const FGeometry& MyGeometry, const FKeyEvent& KeyPressed)
+{
+	if (KeyPressed.GetKey().GetFName() == TEXT("Enter"))
+	{
+		const FText CommittedText = InputText->GetText();
+		bIgnoreUIUpdate = true;
+		InputText->SetText(FText::GetEmpty());
+		SuggestionBox->SetIsOpen( false );
+		bIgnoreUIUpdate = false;
+		return MainPanelWidget.Pin()->ValidateConsoleInputAndAddToCurrentPreset(CommittedText);
+	}
+
+	return FReply::Unhandled();
+}
+
+void SConsoleVariablesEditorCustomConsoleInputBox::SetSuggestions(TArray<FString>& Elements, FText Highlight)
+{
+	FString SelectionText;
+	if (Suggestions.HasSelectedSuggestion())
+	{
+		SelectionText = *Suggestions.GetSelectedSuggestion();
+	}
+
+	Suggestions.Reset();
+	Suggestions.SuggestionsHighlight = Highlight;
+
+	for(int32 i = 0; i < Elements.Num(); ++i)
+	{
+		Suggestions.SuggestionsList.Add(MakeShared<FString>(Elements[i]));
+
+		if (Elements[i] == SelectionText)
+		{
+			Suggestions.SelectedSuggestion = i;
+		}
+	}
+	SuggestionListView->RequestListRefresh();
+
+	if(Suggestions.HasSuggestions())
+	{
+		// Ideally if the selection box is open the output window is not changing it's window title (flickers)
+		SuggestionBox->SetIsOpen(true, false);
+		if (Suggestions.HasSelectedSuggestion())
+		{
+			SuggestionListView->RequestScrollIntoView(Suggestions.GetSelectedSuggestion());
+		}
+		else
+		{
+			SuggestionListView->ScrollToTop();
+		}
+	}
+	else
+	{
+		SuggestionBox->SetIsOpen(false);
+	}
+}
+
+void SConsoleVariablesEditorCustomConsoleInputBox::MarkActiveSuggestion()
+{
+	bIgnoreUIUpdate = true;
+	if (Suggestions.HasSelectedSuggestion())
+	{
+		TSharedPtr<FString> SelectedSuggestion = Suggestions.GetSelectedSuggestion();
+
+		SuggestionListView->SetSelection(SelectedSuggestion);
+		SuggestionListView->RequestScrollIntoView(SelectedSuggestion);	// Ideally this would only scroll if outside of the view
+
+		InputText->SetText(FText::FromString(*SelectedSuggestion));
+	}
+	else
+	{
+		SuggestionListView->ClearSelection();
+	}
+	bIgnoreUIUpdate = false;
+}
+
+void SConsoleVariablesEditorCustomConsoleInputBox::ClearSuggestions()
+{
+	SuggestionBox->SetIsOpen(false);
+	Suggestions.Reset();
+}
+
+#undef LOCTEXT_NAMESPACE
diff --git a/Engine/Plugins/Editor/ConsoleVariablesEditor/Source/ConsoleVariablesEditor/Private/Views/MainPanel/SConsoleVariablesEditorCustomConsoleInputBox.h b/Engine/Plugins/Editor/ConsoleVariablesEditor/Source/ConsoleVariablesEditor/Private/Views/MainPanel/SConsoleVariablesEditorCustomConsoleInputBox.h
new file mode 100644
--- /dev/null
+++ b/Engine/Plugins/Editor/ConsoleVariablesEditor/Source/ConsoleVariablesEditor/Private/Views/MainPanel/SConsoleVariablesEditorCustomConsoleInputBox.h
@@ -0,0 +1,113 @@
+// Copyright Epic Games, Inc. All Rights Reserved.
+
+#pragma once
+
+#include "Widgets/SCompoundWidget.h"
+#include "Widgets/Input/SMenuAnchor.h"
+#include "Widgets/Views/SListView.h"
+
+class SConsoleVariablesEditorMainPanel;
+class SEditableTextBox;
+
+class SConsoleVariablesEditorCustomConsoleInputBox : public SCompoundWidget
+{
+public:
+
+	SLATE_BEGIN_ARGS(SConsoleVariablesEditorCustomConsoleInputBox)
+	{}
+
+	SLATE_END_ARGS()
+	
+	struct FSuggestions
+	{
+		FSuggestions()
+			: SelectedSuggestion(INDEX_NONE)
+		{
+		}
+
+		void Reset()
+		{
+			SelectedSuggestion = INDEX_NONE;
+			SuggestionsList.Reset();
+			SuggestionsHighlight = FText::GetEmpty();
+		}
+
+		bool HasSuggestions() const
+		{
+			return SuggestionsList.Num() > 0;
+		}
+
+		bool HasSelectedSuggestion() const
+		{
+			return SuggestionsList.IsValidIndex(SelectedSuggestion);
+		}
+
+		void StepSelectedSuggestion(const int32 Step)
+		{
+			SelectedSuggestion += Step;
+			if (SelectedSuggestion < 0)
+			{
+				SelectedSuggestion = SuggestionsList.Num() - 1;
+			}
+			else if (SelectedSuggestion >= SuggestionsList.Num())
+			{
+				SelectedSuggestion = 0;
+			}
+		}
+
+		TSharedPtr<FString> GetSelectedSuggestion() const
+		{
+			return SuggestionsList.IsValidIndex(SelectedSuggestion) ? SuggestionsList[SelectedSuggestion] : nullptr;
+		}
+
+		/** INDEX_NONE if not set, otherwise index into SuggestionsList */
+		int32 SelectedSuggestion;
+
+		/** All log messages stored in this widget for the list view */
+		TArray<TSharedPtr<FString>> SuggestionsList;
+
+		/** Highlight text to use for the suggestions list */
+		FText SuggestionsHighlight;
+	};
+
+	void Construct(const FArguments& InArgs, TWeakPtr<SConsoleVariablesEditorMainPanel> InMainPanelWidget);
+
+	virtual ~SConsoleVariablesEditorCustomConsoleInputBox() override;
+
+	//~ Begin SWidget Interface
+	virtual FReply OnPreviewKeyDown(const FGeometry& MyGeometry, const FKeyEvent& KeyEvent) override;
+	//~ End SWidget Interface
+
+	bool TakeKeyboardFocus() const;
+
+	void OnTextChanged(const FText& InText);
+
+	FReply OnKeyCharHandler(const FGeometry& MyGeometry, const FCharacterEvent& InCharacterEvent) const;
+
+	FReply OnKeyDownHandler(const FGeometry& MyGeometry, const FKeyEvent& KeyPressed);
+	
+	void SetSuggestions(TArray<FString>& Elements, FText Highlight);
+
+	void MarkActiveSuggestion();
+
+	void ClearSuggestions();
+
+private:
+
+	TWeakPtr<SConsoleVariablesEditorMainPanel> MainPanelWidget;
+	
+	/** A reference to the actual text box inside ConsoleInput */
+	TSharedPtr<SEditableTextBox> InputText;
+
+	/** history / auto completion elements */
+	TSharedPtr<SMenuAnchor> SuggestionBox;
+
+	/** The list view for showing all log messages. Should be replaced by a full text editor */
+	TSharedPtr<SListView<TSharedPtr<FString>>> SuggestionListView;
+
+	/** Active list of suggestions */
+	FSuggestions Suggestions;
+
+	/** to prevent recursive calls in UI callback */
+	bool bIgnoreUIUpdate = false;
+};
diff --git a/Engine/Plugins/Editor/ConsoleVariablesEditor/Source/ConsoleVariablesEditor/Private/Views/MainPanel/SConsoleVariablesEditorMainPanel.cpp b/Engine/Plugins/Editor/ConsoleVariablesEditor/Source/ConsoleVariablesEditor/Private/Views/MainPanel/SConsoleVariablesEditorMainPanel.cpp
new file mode 100644
--- /dev/null
+++ b/Engine/Plugins/Editor/ConsoleVariablesEditor/Source/ConsoleVariablesEditor/Private/Views/MainPanel/SConsoleVariablesEditorMainPanel.cpp
@@ -0,0 +1,327 @@
+// Copyright Epic Games, Inc. All Rights Reserved.
+
+#include "Views/MainPanel/SConsoleVariablesEditorMainPanel.h"
+
+#include "ConsoleVariablesAsset.h"
+#include "ConsoleVariablesEditorLog.h"
+#include "ConsoleVariablesEditorModule.h"
+#include "ConsoleVariablesEditorStyle.h"
+#include "MultiUser/ConsoleVariableSyncData.h"
+#include "Views/List/ConsoleVariablesEditorList.h"
+#include "Views/MainPanel/ConsoleVariablesEditorMainPanel.h"
+#include "Views/MainPanel/SConsoleVariablesEditorCustomConsoleInputBox.h"
+
+#include "ContentBrowserModule.h"
+#include "Framework/Application/SlateApplication.h"
+#include "Framework/MultiBox/MultiBoxBuilder.h"
+#include "IContentBrowserSingleton.h"
+#include "SPositiveActionButton.h"
+#include "Styling/AppStyle.h"
+#include "Styling/StyleColors.h"
+#include "Widgets/Images/SImage.h"
+#include "Widgets/Input/SButton.h"
+#include "Widgets/Input/SComboButton.h"
+#include "Widgets/Input/SMultiLineEditableTextBox.h"
+#include "Widgets/Layout/SBox.h"
+#include "Widgets/Layout/SSplitter.h"
+#include "Widgets/Text/STextBlock.h"
+
+#define LOCTEXT_NAMESPACE "ConsoleVariablesEditor"
+
+void SConsoleVariablesEditorMainPanel::Construct(
+	const FArguments& InArgs, const TSharedRef<FConsoleVariablesEditorMainPanel>& InMainPanel)
+{
+	check(InMainPanel->GetEditorList().IsValid());
+
+	MainPanel = InMainPanel;
+	
+	ChildSlot
+	[
+		SNew(SVerticalBox)
+			
+		+SVerticalBox::Slot()
+		.HAlign(HAlign_Fill)
+		.AutoHeight()
+		.Padding(FMargin(8.f, 0.f, 8.f, 0.f))
+		[
+			GeneratePanelToolbar()
+		]
+
+		+SVerticalBox::Slot()
+		.HAlign(HAlign_Fill)
+		.AutoHeight()
+		.Padding(FMargin(8.f, 0.f, 8.f, 0.f))
+		[
+			SAssignNew(AddConsoleObjectInputBoxPtr, SConsoleVariablesEditorCustomConsoleInputBox, SharedThis(this))
+			.Visibility(EVisibility::Collapsed)
+		]
+
+		+SVerticalBox::Slot()
+		[
+			SNew(SSplitter)
+			.Orientation(Orient_Vertical)
+			.ResizeMode(ESplitterResizeMode::FixedSize)
+
+			+SSplitter::Slot()
+			[
+				MainPanel.Pin()->GetEditorList().Pin()->GetOrCreateWidget()
+			]
+
+			+SSplitter::Slot()
+			[
+				SAssignNew(MultiUserDetailsBox, SVerticalBox)
+				.Visibility(EVisibility::Collapsed)
+			]
+		]
+	];
+
+	CreateConcertButtonIfNeeded();
+}
+
+SConsoleVariablesEditorMainPanel::~SConsoleVariablesEditorMainPanel()
+{
+	MainPanel.Reset();
+	ToolbarHBox.Reset();
+	ConcertButtonPtr.Reset();
+	MultiUserDetailsBox.Reset();
+}
+
+FReply SConsoleVariablesEditorMainPanel::ValidateConsoleInputAndAddToCurrentPreset(const FText& CommittedText) const
+{
+	return MainPanel.Pin()->ValidateConsoleInputAndAddToCurrentPreset(CommittedText);
+}
+
+void SConsoleVariablesEditorMainPanel::RefreshMultiUserDetails() const
+{
+	UConcertCVarSynchronization* CVarSync = GetMutableDefault<UConcertCVarSynchronization>();
+
+	UConcertCVarConfig* CVarConfig = GetMutableDefault<UConcertCVarConfig>();
+
+	MultiUserDetailsBox->ClearChildren();
+
+	MultiUserDetailsBox->AddSlot()
+	.AutoHeight()
+	[
+		GetConcertDetailsWidget(CVarSync)
+	];
+
+	MultiUserDetailsBox->AddSlot()
+	[
+		GetConcertDetailsWidget(CVarConfig)
+	];
+}
+
+void SConsoleVariablesEditorMainPanel::ToggleMultiUserDetails(ECheckBoxState CheckState) const
+{
+	const bool bShouldBeVisible = CheckState == ECheckBoxState::Checked;
+	MultiUserDetailsBox->SetVisibility(bShouldBeVisible ? EVisibility::SelfHitTestInvisible : EVisibility::Collapsed);
+
+	if (bShouldBeVisible)
+	{
+		RefreshMultiUserDetails();
+	}
+}
+
+TSharedRef<SWidget> SConsoleVariablesEditorMainPanel::GeneratePanelToolbar()
+{
+	return SAssignNew(ToolbarHBox, SHorizontalBox)
+				
+			// Add Console Variable button
+			+ SHorizontalBox::Slot()
+			.VAlign(VAlign_Fill)
+			.HAlign(HAlign_Left)
+			.Padding(FMargin(0, 4))
+			[
+				SNew(SPositiveActionButton)
+				.Text(LOCTEXT("ConsoleVariable", "Console Variable"))
+				.OnClicked(this, &SConsoleVariablesEditorMainPanel::HandleAddConsoleVariableButtonClicked)
+			]
+
+			// Presets Management Button
+			+ SHorizontalBox::Slot()
+			.VAlign(VAlign_Center)
+			.HAlign(HAlign_Right)
+			.Padding(FMargin(8.f, 4, 0.f, 4))
+			[
+				SNew(SComboButton)
+				.ToolTipText(LOCTEXT("PresetManagementButton_Tooltip", "Export the current CVar list to a preset, or import a copy of an existing preset."))
+				.ContentPadding(FMargin(4, 0.5f))
+				.ComboButtonStyle(&FAppStyle::Get().GetWidgetStyle<FComboButtonStyle>("ComboButton"))
+				.OnGetMenuContent(this, &SConsoleVariablesEditorMainPanel::OnGeneratePresetsMenu)
+				.ForegroundColor(FStyleColors::Foreground)
+				.ButtonContent()
+				[
+					SNew(SHorizontalBox)
+
+					+ SHorizontalBox::Slot()
+					.Padding(0, 1, 4, 0)
+					.AutoWidth()
+					[
+						SNew(SImage)
+						.Image(FAppStyle::Get().GetBrush("AssetEditor.SaveAsset"))
+						.ColorAndOpacity(FSlateColor::UseForeground())
+					]
+
+					+ SHorizontalBox::Slot()
+					.Padding(0, 1, 0, 0)
+					.AutoWidth()
+					[
+						SNew(STextBlock)
+						.Text(LOCTEXT("PresetsToolbarButton", "Presets"))
+					]
+				]
+			];
+}
+
+FReply SConsoleVariablesEditorMainPanel::HandleAddConsoleVariableButtonClicked()
+{
+	if (AddConsoleObjectInputBoxPtr.IsValid())
+	{
+		AddConsoleObjectInputBoxPtr->SetVisibility(EVisibility::Visible);
+
+		AddConsoleObjectInputBoxPtr->TakeKeyboardFocus();
+
+		return FReply::Handled();
+	}
+
+	return FReply::Unhandled();
+}
+
+void SConsoleVariablesEditorMainPanel::CreateConcertButtonIfNeeded()
+{
+	if (MainPanel.Pin()->GetMultiUserManager().IsInitialized())
+	{		
+		// Toggle Multi-User Details
+		ToolbarHBox->AddSlot()
+		.HAlign(HAlign_Right)
+		.VAlign(VAlign_Center)
+		.AutoWidth()
+		.Padding(11.f, 1.f, 1.f, 1.f)
+		[
+			SAssignNew(ConcertButtonPtr, SCheckBox)
+			.Padding(0)
+			.ToolTipText(LOCTEXT("ShowConcertSettings_Tip", "Show the multi-user controls for Console Variables"))
+			.Style(&FAppStyle::Get().GetWidgetStyle<FCheckBoxStyle>("ToggleButtonCheckbox"))
+			.ForegroundColor(FSlateColor::UseForeground())
+			.IsChecked(false)
+			.OnCheckStateChanged_Raw(this, &SConsoleVariablesEditorMainPanel::ToggleMultiUserDetails)
+			[
+				SNew(SImage)
+				.Image(FConsoleVariablesEditorStyle::Get().GetBrush("Concert.MultiUser"))
+				.ColorAndOpacity(FSlateColor::UseForeground())
+			]
+		];
+	}
+}
+
+TSharedRef<SWidget> SConsoleVariablesEditorMainPanel::OnGeneratePresetsMenu()
+{
+	FMenuBuilder MenuBuilder(true, NULL);
+
+	IContentBrowserSingleton& ContentBrowser =
+		FModuleManager::LoadModuleChecked<FContentBrowserModule>("ContentBrowser").Get();
+
+	const FText LoadedPresetName = MainPanel.Pin()->GetReferenceAssetOnDisk().IsValid() ?
+		FText::Format(
+			LoadedPresetFormatText,
+			FText::FromString(MainPanel.Pin()->GetReferenceAssetOnDisk()->GetName())) : NoLoadedPresetText;
+
+	FUIAction PresetNameAction = FUIAction();
+	PresetNameAction.CanExecuteAction = FCanExecuteAction::CreateLambda([]() { return false; });
+	
+	MenuBuilder.AddMenuEntry(
+		LoadedPresetName,
+		LoadedPresetName,
+		FSlateIcon(),
+		PresetNameAction,
+		NAME_None,
+		EUserInterfaceActionType::None
+	);
+
+	MenuBuilder.AddMenuSeparator();
+	
+	MenuBuilder.AddMenuEntry(
+		LOCTEXT("SavePreset_Text", "Save Preset"),
+		LOCTEXT("SavePreset_Tooltip", "Save the current preset if one has been loaded. Otherwise, the Save As dialog will be opened."),
+		FSlateIcon(FAppStyle::Get().GetStyleSetName(), "AssetEditor.SaveAsset"),
+		FUIAction(FExecuteAction::CreateRaw(MainPanel.Pin().Get(), &FConsoleVariablesEditorMainPanel::SaveCurrentPreset)),
+		NAME_None,
+		EUserInterfaceActionType::Button
+	);
+
+	MenuBuilder.AddMenuEntry(
+		LOCTEXT("SavePresetAs_Text", "Save Preset As"),
+		LOCTEXT("SavePresetAs_Tooltip", "Save the current configuration as a new preset that can be shared between multiple jobs, or imported later as the base of a new configuration."),
+		FSlateIcon(FAppStyle::Get().GetStyleSetName(), "AssetEditor.SaveAssetAs"),
+		FUIAction(FExecuteAction::CreateRaw(MainPanel.Pin().Get(), &FConsoleVariablesEditorMainPanel::SaveCurrentPresetAs)),
+		NAME_None,
+		EUserInterfaceActionType::Button
+	);
+
+	FAssetPickerConfig AssetPickerConfig;
+	{
+		AssetPickerConfig.SelectionMode = ESelectionMode::Single;
+		AssetPickerConfig.InitialAssetViewType = EAssetViewType::Column;
+		AssetPickerConfig.bFocusSearchBoxWhenOpened = true;
+		AssetPickerConfig.bAllowNullSelection = false;
+		AssetPickerConfig.bShowBottomToolbar = true;
+		AssetPickerConfig.bAutohideSearchBar = false;
+		AssetPickerConfig.bAllowDragging = false;
+		AssetPickerConfig.bCanShowClasses = false;
+		AssetPickerConfig.bShowPathInColumnView = true;
+		AssetPickerConfig.bShowTypeInColumnView = false;
+		AssetPickerConfig.bSortByPathInColumnView = false;
+		AssetPickerConfig.bForceShowEngineContent = false;
+		AssetPickerConfig.bForceShowPluginContent = false;
+
+		AssetPickerConfig.AssetShowWarningText = LOCTEXT("NoPresets_Warning", "No Presets Found");
+		AssetPickerConfig.Filter.ClassNames.Add(UConsoleVariablesAsset::StaticClass()->GetFName());
+		AssetPickerConfig.Filter.bRecursiveClasses = true;
+		AssetPickerConfig.OnAssetSelected =
+			FOnAssetSelected::CreateRaw(MainPanel.Pin().Get(), &FConsoleVariablesEditorMainPanel::ImportPreset);
+	}
+
+	MenuBuilder.BeginSection(NAME_None, LOCTEXT("ImportPreset_MenuSection", "Import Preset"));
+	{
+		TSharedRef<SWidget> PresetPicker = SNew(SBox)
+			.MinDesiredWidth(400.f)
+			.MinDesiredHeight(400.f)
+			[
+				ContentBrowser.CreateAssetPicker(AssetPickerConfig)
+			];
+
+		MenuBuilder.AddWidget(PresetPicker, FText(), true, false);
+	}
+	MenuBuilder.EndSection();
+
+	TSharedRef<SWidget> Widget = MenuBuilder.MakeWidget();
+	FChildren* ChildWidgets = Widget->GetChildren();
+	for (int32 ChildItr = 0; ChildItr < ChildWidgets->Num(); ChildItr++)
+	{
+		TSharedRef<SWidget> Child = ChildWidgets->GetChildAt(ChildItr);
+
+		Child->EnableToolTipForceField(false);
+	}
+	Widget->EnableToolTipForceField(false);
+	
+	return Widget;
+}
+
+TSharedRef<SWidget> SConsoleVariablesEditorMainPanel::GetConcertDetailsWidget(UObject* InObject)
+{
+	FPropertyEditorModule& PropertyEditorModule =
+		FModuleManager::Get().LoadModuleChecked<FPropertyEditorModule>("PropertyEditor");
+	FDetailsViewArgs DetailsViewArgs;
+	DetailsViewArgs.bAllowSearch = false;
+	DetailsViewArgs.NameAreaSettings = FDetailsViewArgs::HideNameArea;
+	DetailsViewArgs.bHideSelectionTip = true;
+	DetailsViewArgs.bShowScrollBar = false;
+
+	TSharedRef<IDetailsView> Details = PropertyEditorModule.CreateDetailView(DetailsViewArgs);
+
+	Details->SetObjects(TArray{InObject});
+
+	return Details;
+}
+
+#undef LOCTEXT_NAMESPACE
diff --git a/Engine/Plugins/Editor/ConsoleVariablesEditor/Source/ConsoleVariablesEditor/Private/Views/MainPanel/SConsoleVariablesEditorMainPanel.h b/Engine/Plugins/Editor/ConsoleVariablesEditor/Source/ConsoleVariablesEditor/Private/Views/MainPanel/SConsoleVariablesEditorMainPanel.h
new file mode 100644
--- /dev/null
+++ b/Engine/Plugins/Editor/ConsoleVariablesEditor/Source/ConsoleVariablesEditor/Private/Views/MainPanel/SConsoleVariablesEditorMainPanel.h
@@ -0,0 +1,70 @@
+// Copyright Epic Games, Inc. All Rights Reserved.
+
+#pragma once
+
+#include "Widgets/Input/SCheckBox.h"
+#include "Widgets/SCompoundWidget.h"
+
+class FConsoleVariablesEditorMainPanel;
+class SBox;
+class SConsoleVariablesEditorCustomConsoleInputBox;
+class SHorizontalBox;
+class SVerticalBox;
+class UConsoleVariablesAsset;
+
+class SConsoleVariablesEditorMainPanel final : public SCompoundWidget
+{
+
+public:
+
+	SLATE_BEGIN_ARGS(SConsoleVariablesEditorMainPanel)
+	{}
+
+	SLATE_END_ARGS()
+
+	void Construct(const FArguments& InArgs, const TSharedRef<FConsoleVariablesEditorMainPanel>& InMainPanel);
+
+	virtual ~SConsoleVariablesEditorMainPanel() override;
+
+	/** A handler for when text is committed to the Console Input Box that appears when clicking the "Add Console Variable" button */
+	FReply ValidateConsoleInputAndAddToCurrentPreset(const FText& CommittedText) const;
+
+	void RefreshMultiUserDetails() const;
+	void ToggleMultiUserDetails(ECheckBoxState CheckState) const;
+
+private:
+
+	/** A reference to the struct that controls this widget */
+	TWeakPtr<FConsoleVariablesEditorMainPanel> MainPanel;
+
+	TSharedPtr<SConsoleVariablesEditorCustomConsoleInputBox> AddConsoleObjectInputBoxPtr;
+
+	TSharedPtr<SHorizontalBox> ToolbarHBox;
+
+	/** A reference to the button which opens the plugin settings */
+	TSharedPtr<SCheckBox> ConcertButtonPtr;
+
+	/** Creates the toolbar at the top of the MainPanel widget */
+	TSharedRef<SWidget> GeneratePanelToolbar();
+
+	FReply HandleAddConsoleVariableButtonClicked();
+	
+	void CreateConcertButtonIfNeeded();
+
+	// Save / Load
+
+	#define LOCTEXT_NAMESPACE "ConsoleVariablesEditor"
+
+	const FText NoLoadedPresetText = LOCTEXT("NoLoadedPreset", "No Loaded Preset");
+	const FText LoadedPresetFormatText = LOCTEXT("LoadedPresetFormat", "Current Preset: {0}");
+
+    #undef LOCTEXT_NAMESPACE
+
+	/** Creates a special asset picker widget to display when the Save/Load button is clicked */
+	TSharedRef<SWidget> OnGeneratePresetsMenu();
+
+	TSharedPtr<SVerticalBox> MultiUserDetailsBox;
+
+	/** Generates a customized details widget given an object. Intended for multi-user settings display. */
+	static TSharedRef<SWidget> GetConcertDetailsWidget(UObject* InObject);
+};
diff --git a/Engine/Plugins/Editor/ConsoleVariablesEditor/Source/ConsoleVariablesEditor/Private/Views/Widgets/SConsoleVariablesEditorGlobalSearchToggle.cpp b/Engine/Plugins/Editor/ConsoleVariablesEditor/Source/ConsoleVariablesEditor/Private/Views/Widgets/SConsoleVariablesEditorGlobalSearchToggle.cpp
new file mode 100644
--- /dev/null
+++ b/Engine/Plugins/Editor/ConsoleVariablesEditor/Source/ConsoleVariablesEditor/Private/Views/Widgets/SConsoleVariablesEditorGlobalSearchToggle.cpp
@@ -0,0 +1,141 @@
+﻿// Copyright Epic Games, Inc. All Rights Reserved.
+
+#include "SConsoleVariablesEditorGlobalSearchToggle.h"
+
+#include "Framework/Application/SlateApplication.h"
+#include "Styling/AppStyle.h"
+#include "Styling/StyleColors.h"
+#include "Widgets/Images/SImage.h"
+#include "Widgets/Layout/SBorder.h"
+#include "Widgets/SBoxPanel.h"
+#include "Widgets/Text/STextBlock.h"
+
+#define LOCTEXT_NAMESPACE "ConsoleVariablesEditor"
+
+void SConsoleVariablesEditorGlobalSearchToggle::Construct(const FArguments& Args, const FText& InButtonText)
+{
+	OnToggleCtrlClicked = Args._OnToggleCtrlClicked;
+	OnToggleAltClicked = Args._OnToggleAltClicked;
+	OnToggleMiddleButtonClicked = Args._OnToggleMiddleButtonClicked;
+	OnToggleRightButtonClicked = Args._OnToggleRightButtonClicked;
+	OnToggleClickedOnce = Args._OnToggleClickedOnce;
+
+	GlobalSearchText = InButtonText;
+	CheckedColor = USlateThemeManager::Get().GetColor(EStyleColor::Primary);
+	UncheckedColor = USlateThemeManager::Get().GetColor(EStyleColor::Black);
+	
+	ChildSlot
+    [
+		SNew(SBorder)
+		.Padding(1.0f)
+		.BorderImage(FAppStyle::Get().GetBrush("ContentBrowser.FilterBackground"))
+		 [
+			SAssignNew( ToggleButtonPtr, SCheckBox )
+			.Style(FAppStyle::Get(), "ContentBrowser.FilterButton")
+			.ToolTipText(LOCTEXT("GlobalSearchPillButtonTooltip","Toggle this search"))
+			.Padding(0.0f)
+			.IsChecked(this, &SConsoleVariablesEditorGlobalSearchToggle::GetCheckedState)
+			.OnCheckStateChanged_Lambda(
+				[this](ECheckBoxState NewCheckboxState)
+			{
+				SetIsButtonChecked(!GetIsToggleChecked());
+
+				if (FSlateApplication::Get().GetModifierKeys().IsControlDown() && OnToggleCtrlClicked.IsBound())
+				{
+					SetIsMarkedForDelete(true);
+					OnToggleCtrlClicked.Execute();
+				}
+				else if (FSlateApplication::Get().GetModifierKeys().IsAltDown() && OnToggleAltClicked.IsBound())
+				{
+					SetIsMarkedForDelete(true);
+					OnToggleAltClicked.Execute();
+				}
+				else if (OnToggleClickedOnce.IsBound())
+				{
+					SetIsMarkedForDelete(false);
+					OnToggleClickedOnce.Execute();
+				}
+			})
+			[
+				SNew(SHorizontalBox)
+				+SHorizontalBox::Slot()
+				.VAlign(VAlign_Center)
+				.AutoWidth()
+				[
+					SNew(SImage)
+					.Image(FAppStyle::Get().GetBrush("ContentBrowser.FilterImage"))
+					.ColorAndOpacity_Lambda([this]()
+					{
+						return GetIsToggleChecked() ? CheckedColor : UncheckedColor;
+					})
+				]
+				+SHorizontalBox::Slot()
+				.Padding(2.f)
+				.VAlign(VAlign_Center)
+				[
+					SNew(STextBlock)
+					.Text(FText::Format(INVTEXT("\"{0}\""), GlobalSearchText)) // Show text in quotes
+				]
+			]
+		]
+    ];
+}
+
+SConsoleVariablesEditorGlobalSearchToggle::~SConsoleVariablesEditorGlobalSearchToggle()
+{
+	OnToggleCtrlClicked.Unbind();
+	OnToggleAltClicked.Unbind();
+	OnToggleMiddleButtonClicked.Unbind();
+	OnToggleRightButtonClicked.Unbind();
+	OnToggleClickedOnce.Unbind();
+
+	ToggleButtonPtr.Reset();
+}
+
+FReply SConsoleVariablesEditorGlobalSearchToggle::OnMouseButtonDown(const FGeometry& MyGeometry, const FPointerEvent& MouseEvent)
+{
+	return FReply::Handled();
+}
+
+FReply SConsoleVariablesEditorGlobalSearchToggle::OnMouseButtonUp( const FGeometry& InMyGeometry, const FPointerEvent& InMouseEvent) 
+{
+	if(InMouseEvent.GetEffectingButton() == EKeys::MiddleMouseButton && OnToggleMiddleButtonClicked.IsBound())
+	{
+		SetIsMarkedForDelete(true);
+		return OnToggleMiddleButtonClicked.Execute().ReleaseMouseCapture();
+	}
+	else if (InMouseEvent.GetEffectingButton() == EKeys::RightMouseButton && OnToggleRightButtonClicked.IsBound())
+	{
+		SetIsMarkedForDelete(true);
+		return OnToggleRightButtonClicked.Execute().ReleaseMouseCapture();
+	}
+		
+	return FReply::Handled().ReleaseMouseCapture();
+}
+
+bool SConsoleVariablesEditorGlobalSearchToggle::GetIsToggleChecked() const
+{
+	return bIsToggleChecked;
+}
+
+ECheckBoxState SConsoleVariablesEditorGlobalSearchToggle::GetCheckedState() const
+{
+	return GetIsToggleChecked() ? ECheckBoxState::Checked : ECheckBoxState::Unchecked;
+}
+
+void SConsoleVariablesEditorGlobalSearchToggle::SetIsButtonChecked(const bool bNewIsButtonChecked)
+{
+	bIsToggleChecked = bNewIsButtonChecked;
+}
+
+bool SConsoleVariablesEditorGlobalSearchToggle::GetIsMarkedForDelete() const
+{
+	return bIsMarkedForDelete;
+}
+
+void SConsoleVariablesEditorGlobalSearchToggle::SetIsMarkedForDelete(const bool bNewMark)
+{
+	bIsMarkedForDelete = bNewMark;
+}
+
+#undef LOCTEXT_NAMESPACE
diff --git a/Engine/Plugins/Editor/ConsoleVariablesEditor/Source/ConsoleVariablesEditor/Private/Views/Widgets/SConsoleVariablesEditorGlobalSearchToggle.h b/Engine/Plugins/Editor/ConsoleVariablesEditor/Source/ConsoleVariablesEditor/Private/Views/Widgets/SConsoleVariablesEditorGlobalSearchToggle.h
new file mode 100644
--- /dev/null
+++ b/Engine/Plugins/Editor/ConsoleVariablesEditor/Source/ConsoleVariablesEditor/Private/Views/Widgets/SConsoleVariablesEditorGlobalSearchToggle.h
@@ -0,0 +1,64 @@
+﻿// Copyright Epic Games, Inc. All Rights Reserved.
+
+#pragma once
+
+#include "CoreMinimal.h"
+
+#include "Framework/SlateDelegates.h"
+#include "Widgets/DeclarativeSyntaxSupport.h"
+#include "Widgets/Input/SCheckBox.h"
+#include "Widgets/SCompoundWidget.h"
+
+class SConsoleVariablesEditorGlobalSearchToggle : public SCompoundWidget
+{
+public:
+
+	SLATE_BEGIN_ARGS(SConsoleVariablesEditorGlobalSearchToggle)
+	{}
+		SLATE_EVENT(FOnClicked, OnToggleCtrlClicked)
+		SLATE_EVENT(FOnClicked, OnToggleAltClicked)
+		SLATE_EVENT(FOnClicked, OnToggleMiddleButtonClicked)
+		SLATE_EVENT(FOnClicked, OnToggleRightButtonClicked)
+		SLATE_EVENT(FOnClicked, OnToggleClickedOnce)
+	SLATE_END_ARGS();
+
+	void Construct(const FArguments& Args, const FText& InButtonText);
+
+	virtual ~SConsoleVariablesEditorGlobalSearchToggle() override;
+	
+	virtual FReply OnMouseButtonDown(const FGeometry& MyGeometry, const FPointerEvent& MouseEvent) override;
+	virtual FReply OnMouseButtonUp( const FGeometry& InMyGeometry, const FPointerEvent& InMouseEvent ) override;
+
+	
+	[[nodiscard]] bool GetIsToggleChecked() const;
+	[[nodiscard]] ECheckBoxState GetCheckedState() const;
+	void SetIsButtonChecked(const bool bNewIsButtonChecked);
+
+	[[nodiscard]] bool GetIsMarkedForDelete() const;
+	void SetIsMarkedForDelete(const bool bNewMark);
+
+	const FText& GetGlobalSearchText()
+	{
+		return GlobalSearchText;
+	}
+
+private:
+	
+	FOnClicked OnToggleCtrlClicked;
+	FOnClicked OnToggleAltClicked;
+	FOnClicked OnToggleMiddleButtonClicked;
+	FOnClicked OnToggleRightButtonClicked;
+	FOnClicked OnToggleClickedOnce;
+
+	TSharedPtr<SCheckBox> ToggleButtonPtr;
+	bool bIsToggleChecked = true;
+
+	/** When clicked in a special manner, this search button will be marked for deletion
+	 * The responsibility is on external classes to remove the button from the UI the external class controls. */
+	bool bIsMarkedForDelete = false;
+
+	FLinearColor CheckedColor;
+	FLinearColor UncheckedColor;
+
+	FText GlobalSearchText;
+};
diff --git a/Engine/Plugins/Editor/ConsoleVariablesEditor/Source/ConsoleVariablesEditor/Private/Views/Widgets/SConsoleVariablesEditorTooltipWidget.h b/Engine/Plugins/Editor/ConsoleVariablesEditor/Source/ConsoleVariablesEditor/Private/Views/Widgets/SConsoleVariablesEditorTooltipWidget.h
new file mode 100644
--- /dev/null
+++ b/Engine/Plugins/Editor/ConsoleVariablesEditor/Source/ConsoleVariablesEditor/Private/Views/Widgets/SConsoleVariablesEditorTooltipWidget.h
@@ -0,0 +1,79 @@
+// Copyright Epic Games, Inc. All Rights Reserved.
+
+#pragma once
+
+#include "Styling/AppStyle.h"
+#include "Widgets/Layout/SBorder.h"
+#include "Widgets/SBoxPanel.h"
+#include "Widgets/SCompoundWidget.h"
+#include "Widgets/SToolTip.h"
+#include "Widgets/Text/STextBlock.h"
+
+class SConsoleVariablesEditorTooltipWidget : public SCompoundWidget
+{
+public:
+	
+	SLATE_BEGIN_ARGS(SConsoleVariablesEditorTooltipWidget)
+	{}
+
+	SLATE_END_ARGS()
+	
+	void Construct(const FArguments& InArgs, const FString InCommandName, const FString InHelpString)
+	{
+		TSharedPtr<SVerticalBox> VerticalBox;
+		
+		ChildSlot
+		[
+			SAssignNew(VerticalBox, SVerticalBox)
+
+			+SVerticalBox::Slot()
+			.AutoHeight()
+			.Padding(0, 0, 0, 4)
+			[
+				SNew(SBorder)
+				.Padding(6)
+				.BorderImage(FAppStyle::Get().GetBrush("ContentBrowser.TileViewTooltip.ContentBorder"))
+				[
+					SNew(STextBlock)
+					.Text(FText::FromString(InCommandName))
+					.Font(FAppStyle::Get().GetFontStyle("ContentBrowser.TileViewTooltip.NameFont"))
+					.AutoWrapText(true)
+				]
+			]
+		];
+
+		if (!InHelpString.IsEmpty())
+		{
+			VerticalBox->AddSlot()
+			.AutoHeight()
+			[
+				SNew(SBorder)
+				.Padding(6)
+				.BorderImage(FAppStyle::Get().GetBrush("ContentBrowser.TileViewTooltip.ContentBorder"))
+				[
+					SNew(STextBlock)
+					.Text(FText::FromString(InHelpString))
+					.AutoWrapText(true)
+				]
+			];
+		}
+	}
+
+	static TSharedRef<SToolTip> MakeTooltip(const FString InCommandName, const FString InHelpString)
+	{
+		return
+			SNew(SToolTip)
+			.BorderImage(FAppStyle::Get().GetBrush("ContentBrowser.TileViewTooltip.NonContentBorder"))
+			[
+				SNew(SBox)
+				.WidthOverride(TooltipWindowWidth)
+				[
+					SNew(SConsoleVariablesEditorTooltipWidget, InCommandName, InHelpString)
+				]
+			];
+	}
+
+private:
+
+	static constexpr float TooltipWindowWidth = 500.f;
+};
diff --git a/Engine/Plugins/Editor/ConsoleVariablesEditor/Source/ConsoleVariablesEditor/Public/ConsoleVariablesAsset.h b/Engine/Plugins/Editor/ConsoleVariablesEditor/Source/ConsoleVariablesEditor/Public/ConsoleVariablesAsset.h
new file mode 100644
--- /dev/null
+++ b/Engine/Plugins/Editor/ConsoleVariablesEditor/Source/ConsoleVariablesEditor/Public/ConsoleVariablesAsset.h
@@ -0,0 +1,86 @@
+// Copyright Epic Games, Inc. All Rights Reserved.
+
+#pragma once
+
+#include "ConsoleVariablesEditorCommandInfo.h"
+
+#include "CoreMinimal.h"
+#include "Widgets/Input/SCheckBox.h"
+
+#include "ConsoleVariablesAsset.generated.h"
+
+/** Data that will be serialized with this asset */
+USTRUCT()
+struct FConsoleVariablesEditorAssetSaveData
+{
+	GENERATED_BODY()
+
+	FORCEINLINE bool operator==(const FConsoleVariablesEditorAssetSaveData& Comparator) const
+	{
+		return CommandName.Equals(Comparator.CommandName);
+	}
+	
+	UPROPERTY()
+	FString CommandName;
+
+	UPROPERTY()
+	FString CommandValueAsString;
+
+	UPROPERTY()
+	/** If Undetermined, we can assume this data was not previously saved */
+	ECheckBoxState CheckedState = ECheckBoxState::Undetermined;
+};
+
+/** An asset used to track collections of console variables that can be recalled and edited using the Console Variables Editor. */
+UCLASS(BlueprintType)
+class CONSOLEVARIABLESEDITOR_API UConsoleVariablesAsset : public UObject
+{
+	GENERATED_BODY()
+public:
+
+	/** Sets a description for this variable collection. */
+	UFUNCTION(BlueprintCallable, Category = "Console Variables Editor")
+	void SetVariableCollectionDescription(const FString& InVariableCollectionDescription);
+	UFUNCTION(BlueprintPure, Category = "Console Variables Editor")
+	FString GetVariableCollectionDescription() const
+	{
+		return VariableCollectionDescription;
+	}
+
+	/** Returns the saved list of console variable information such as the variable name, the type and the value of the vriable at the time the asset was saved. */
+	const TArray<FConsoleVariablesEditorAssetSaveData>& GetSavedCommands() const
+	{
+		return SavedCommands;
+	}
+
+	/** Completely replaces the saved data with new saved data */
+	void ReplaceSavedCommands(const TArray<FConsoleVariablesEditorAssetSaveData>& Replacement);
+
+	/** Returns how many console variables are serialized in this asset */
+	int32 GetSavedCommandsCount() const
+	{
+		return GetSavedCommands().Num();
+	}
+
+	/** Outputs the FConsoleVariablesEditorCommandInfo matching InCommand. Returns whether a match was found. Case sensitive. */
+	bool FindSavedDataByCommandString(const FString InCommandString, FConsoleVariablesEditorAssetSaveData& OutValue) const;
+
+	/** Set the value of a saved console variable if the name matches; add a new console variable to the list if a match is not found. */
+	void AddOrSetConsoleObjectSavedData(const FConsoleVariablesEditorAssetSaveData& InData);
+
+	/** Returns true if the element was found and successfully removed. */
+	bool RemoveConsoleVariable(const FString InCommandString);
+
+	/** Copy data from input asset to this asset */
+	void CopyFrom(const UConsoleVariablesAsset* InAssetToCopy);
+	
+private:
+	
+	/* User defined description of the variable collection */
+	UPROPERTY(AssetRegistrySearchable, BlueprintGetter = "GetVariableCollectionDescription", EditAnywhere, Category = "Console Variables Editor")
+	FString VariableCollectionDescription;
+
+	/** A saved list of console variable information such as the variable name, the type and the value of the variable at the time the asset was saved. */
+	UPROPERTY()
+	TArray<FConsoleVariablesEditorAssetSaveData> SavedCommands;
+};
diff --git a/Engine/Plugins/Editor/ConsoleVariablesEditor/Source/ConsoleVariablesEditor/Public/ConsoleVariablesEditorFunctionLibrary.h b/Engine/Plugins/Editor/ConsoleVariablesEditor/Source/ConsoleVariablesEditor/Public/ConsoleVariablesEditorFunctionLibrary.h
new file mode 100644
--- /dev/null
+++ b/Engine/Plugins/Editor/ConsoleVariablesEditor/Source/ConsoleVariablesEditor/Public/ConsoleVariablesEditorFunctionLibrary.h
@@ -0,0 +1,87 @@
+// Copyright Epic Games, Inc. All Rights Reserved.
+
+#pragma once
+
+#include "CoreMinimal.h"
+#include "Kismet/BlueprintFunctionLibrary.h"
+
+#include "ConsoleVariablesEditorFunctionLibrary.generated.h"
+
+class UConsoleVariablesAsset;
+
+/** An asset used to track collections of console variables that can be recalled and edited using the Console Variables Editor. */
+UCLASS(BlueprintType)
+class CONSOLEVARIABLESEDITOR_API UConsoleVariablesEditorFunctionLibrary : public UBlueprintFunctionLibrary
+{
+	GENERATED_BODY()
+public:
+
+	/** Return the currently loaded list of variables in the Console Variables Editor. */
+	UFUNCTION(BlueprintCallable, Category = "Console Variables Editor")
+	static UConsoleVariablesAsset* GetCurrentlyLoadedPreset();
+
+	/** Loads the given asset in the Console Variables Editor and sets all its variable values. */
+	UFUNCTION(BlueprintCallable, Category = "Console Variables Editor")
+	static void LoadPresetIntoConsoleVariablesEditor(const UConsoleVariablesAsset* InAsset);
+
+	/*
+	 * Saves the current list in the Console Variables Editor to the given asset.
+	 * The Asset will not be automatically saved.
+	 * @return true if successful.
+	 */
+	UFUNCTION(BlueprintCallable, Category = "Console Variables Editor")
+	static bool CopyCurrentListToAsset(UConsoleVariablesAsset* InAsset);
+
+	/*
+	 * Adds a validated command to the current preset with its current value.
+	 * The Asset will not be automatically saved.
+	 * @return true if successful.
+	 */
+	UFUNCTION(BlueprintCallable, Category = "Console Variables Editor")
+	static bool AddValidatedCommandToCurrentPreset(const FString NewCommand);
+
+	/*
+	 * Removes a command from the current preset if it exists in the saved data.
+	 * The Asset will not be automatically saved.
+	 * @return true if successful.
+	 */
+	UFUNCTION(BlueprintCallable, Category = "Console Variables Editor")
+	static bool RemoveCommandFromCurrentPreset(const FString NewCommand);
+
+	/** Return an array of strings containing the command names for each command found in the given preset. */
+	UFUNCTION(BlueprintCallable, Category = "Console Variables Editor")
+	static bool GetListOfCommandsFromPreset(const UConsoleVariablesAsset* InAsset, TArray<FString>& OutCommandList);
+
+	/** Set a console variable value directly. Returns true if the console object exists. */
+	UFUNCTION(BlueprintCallable, Category = "Console Variables Editor", meta=(DisplayName="Set Console Variable By Name (Float)"))
+	static bool SetConsoleVariableByName_Float(const FString InCommandName, const float InValue);
+
+	/** Set a console variable value directly. Returns true if the console object exists. */
+	UFUNCTION(BlueprintCallable, Category = "Console Variables Editor", meta=(DisplayName="Set Console Variable By Name (Bool)"))
+	static bool SetConsoleVariableByName_Bool(const FString InCommandName, const bool InValue);
+
+	/** Set a console variable value directly. Returns true if the console object exists. */
+	UFUNCTION(BlueprintCallable, Category = "Console Variables Editor", meta=(DisplayName="Set Console Variable By Name (Int)"))
+	static bool SetConsoleVariableByName_Int(const FString InCommandName, const int32 InValue);
+	
+	/** Set a console variable value directly. Returns true if the console object exists. */
+	UFUNCTION(BlueprintCallable, Category = "Console Variables Editor", meta=(DisplayName="Set Console Variable By Name (String)"))
+	static bool SetConsoleVariableByName_String(const FString InCommandName, const FString InValue);
+
+	/** Get a console variable's string value directly. Returns true if the console object exists. */
+	UFUNCTION(BlueprintCallable, Category = "Console Variables Editor")
+	static bool GetConsoleVariableStringValue(const FString InCommandName, FString& OutValue);
+
+	/** Set a console variable value directly. Returns true if the console object exists. */
+	UFUNCTION(BlueprintCallable, Category = "Console Variables Editor")
+	static bool GetConsoleVariableSourceByName(const FString InCommandName, FString& OutValue);
+
+	/** Return whether the Multi-user sync setting for the current instance of the editor is enabled. */
+	UFUNCTION(BlueprintCallable, Category = "Console Variables Editor")
+	static bool GetEnableMultiUserCVarSync();
+	
+	/** Enable or disable the Multi-user sync setting for the current instance of the editor. */
+	UFUNCTION(BlueprintCallable, Category = "Console Variables Editor")
+	static void SetEnableMultiUserCVarSync(bool bNewSetting);
+	
+};
diff --git a/Engine/Plugins/Editor/ConsoleVariablesEditor/Source/ConsoleVariablesEditor/Public/ConsoleVariablesEditorLog.h b/Engine/Plugins/Editor/ConsoleVariablesEditor/Source/ConsoleVariablesEditor/Public/ConsoleVariablesEditorLog.h
new file mode 100644
--- /dev/null
+++ b/Engine/Plugins/Editor/ConsoleVariablesEditor/Source/ConsoleVariablesEditor/Public/ConsoleVariablesEditorLog.h
@@ -0,0 +1,7 @@
+// Copyright Epic Games, Inc. All Rights Reserved.
+
+#pragma once
+
+#include "CoreMinimal.h"
+
+DECLARE_LOG_CATEGORY_CLASS(LogConsoleVariablesEditor, Log, All);
diff --git a/Engine/Plugins/Editor/ConsoleVariablesEditor/Source/ConsoleVariablesEditor/Public/ConsoleVariablesEditorProjectSettings.h b/Engine/Plugins/Editor/ConsoleVariablesEditor/Source/ConsoleVariablesEditor/Public/ConsoleVariablesEditorProjectSettings.h
new file mode 100644
--- /dev/null
+++ b/Engine/Plugins/Editor/ConsoleVariablesEditor/Source/ConsoleVariablesEditor/Public/ConsoleVariablesEditorProjectSettings.h
@@ -0,0 +1,55 @@
+// Copyright Epic Games, Inc. All Rights Reserved.
+
+#pragma once
+
+#include "ConsoleVariablesEditorProjectSettings.generated.h"
+
+UENUM(BlueprintType)
+enum class EConsoleVariablesEditorRowDisplayType : uint8
+{
+	ShowCurrentValue,
+	ShowLastEnteredValue
+};
+
+UCLASS(config = Engine, defaultconfig)
+class CONSOLEVARIABLESEDITOR_API UConsoleVariablesEditorProjectSettings : public UObject
+{
+	GENERATED_BODY()
+public:
+	
+	UConsoleVariablesEditorProjectSettings(const FObjectInitializer& ObjectInitializer)
+	{
+		UncheckedRowDisplayType = EConsoleVariablesEditorRowDisplayType::ShowCurrentValue;
+		
+		bAddAllChangedConsoleVariablesToCurrentPreset = true;
+
+		ChangedConsoleVariableSkipList =
+		{
+			"Editor.ReflectEditorLevelVisibilityWithGame"
+		};
+	}
+
+	/**
+	 *When a row is unchecked, its associated variable's value will be set to the value recorded when the plugin was loaded.
+	 *The value displayed to the user can be configured with this setting, but will not affect the actual applied value.
+	 *ShowCurrentValue displays the actual value currently applied to the variable.
+	 *ShowLastEnteredValue displays the value that will be applied when the row is checked.
+	 */
+	UPROPERTY(Config, EditAnywhere, Category="Console Variables Editor")
+	EConsoleVariablesEditorRowDisplayType UncheckedRowDisplayType;
+
+	/**
+	 *When variables are changed outside the Console Variables Editor, this option will add the variables to the current preset.
+	 *Does not apply to console commands like 'r.SetNearClipPlane' or 'stat fps'
+	 */
+	UPROPERTY(Config, EditAnywhere, Category="Console Variables Editor")
+	bool bAddAllChangedConsoleVariablesToCurrentPreset;
+
+	/**
+	 * If bAddAllChangedConsoleVariablesToCurrentPreset is true, this list will filter out any matching variables
+	 * changed outside of the Console Variables Editor so they won't be added to the current preset.
+	 * Matching variables explicitly added inside the Console Variables Editor will not be filtered.
+	 */
+	UPROPERTY(Config, EditAnywhere, Category="Console Variables Editor")
+	TSet<FString> ChangedConsoleVariableSkipList;
+};
diff --git a/Engine/Plugins/Experimental/ConsoleVariables/ConsoleVariables.uplugin b/Engine/Plugins/Experimental/ConsoleVariables/ConsoleVariables.uplugin
deleted file mode 100644
diff --git a/Engine/Plugins/Experimental/ConsoleVariables/Source/ConsoleVariablesEditor/ConsoleVariablesEditor.build.cs b/Engine/Plugins/Experimental/ConsoleVariables/Source/ConsoleVariablesEditor/ConsoleVariablesEditor.build.cs
deleted file mode 100644
diff --git a/Engine/Plugins/Experimental/ConsoleVariables/Source/ConsoleVariablesEditor/Private/AssetTypeActions/AssetTypeActions_ConsoleVariables.cpp b/Engine/Plugins/Experimental/ConsoleVariables/Source/ConsoleVariablesEditor/Private/AssetTypeActions/AssetTypeActions_ConsoleVariables.cpp
deleted file mode 100644
diff --git a/Engine/Plugins/Experimental/ConsoleVariables/Source/ConsoleVariablesEditor/Private/AssetTypeActions/AssetTypeActions_ConsoleVariables.h b/Engine/Plugins/Experimental/ConsoleVariables/Source/ConsoleVariablesEditor/Private/AssetTypeActions/AssetTypeActions_ConsoleVariables.h
deleted file mode 100644
diff --git a/Engine/Plugins/Experimental/ConsoleVariables/Source/ConsoleVariablesEditor/Private/ConsoleVariablesAsset.cpp b/Engine/Plugins/Experimental/ConsoleVariables/Source/ConsoleVariablesEditor/Private/ConsoleVariablesAsset.cpp
deleted file mode 100644
diff --git a/Engine/Plugins/Experimental/ConsoleVariables/Source/ConsoleVariablesEditor/Private/ConsoleVariablesEditorCommandInfo.h b/Engine/Plugins/Experimental/ConsoleVariables/Source/ConsoleVariablesEditor/Private/ConsoleVariablesEditorCommandInfo.h
deleted file mode 100644
diff --git a/Engine/Plugins/Experimental/ConsoleVariables/Source/ConsoleVariablesEditor/Private/ConsoleVariablesEditorFunctionLibrary.cpp b/Engine/Plugins/Experimental/ConsoleVariables/Source/ConsoleVariablesEditor/Private/ConsoleVariablesEditorFunctionLibrary.cpp
deleted file mode 100644
diff --git a/Engine/Plugins/Experimental/ConsoleVariables/Source/ConsoleVariablesEditor/Private/ConsoleVariablesEditorModule.cpp b/Engine/Plugins/Experimental/ConsoleVariables/Source/ConsoleVariablesEditor/Private/ConsoleVariablesEditorModule.cpp
deleted file mode 100644
diff --git a/Engine/Plugins/Experimental/ConsoleVariables/Source/ConsoleVariablesEditor/Private/ConsoleVariablesEditorModule.h b/Engine/Plugins/Experimental/ConsoleVariables/Source/ConsoleVariablesEditor/Private/ConsoleVariablesEditorModule.h
deleted file mode 100644
diff --git a/Engine/Plugins/Experimental/ConsoleVariables/Source/ConsoleVariablesEditor/Private/ConsoleVariablesEditorStyle.cpp b/Engine/Plugins/Experimental/ConsoleVariables/Source/ConsoleVariablesEditor/Private/ConsoleVariablesEditorStyle.cpp
deleted file mode 100644
diff --git a/Engine/Plugins/Experimental/ConsoleVariables/Source/ConsoleVariablesEditor/Private/ConsoleVariablesEditorStyle.h b/Engine/Plugins/Experimental/ConsoleVariables/Source/ConsoleVariablesEditor/Private/ConsoleVariablesEditorStyle.h
deleted file mode 100644
diff --git a/Engine/Plugins/Experimental/ConsoleVariables/Source/ConsoleVariablesEditor/Private/Factories/ConsoleVariablesEditorFactory.cpp b/Engine/Plugins/Experimental/ConsoleVariables/Source/ConsoleVariablesEditor/Private/Factories/ConsoleVariablesEditorFactory.cpp
deleted file mode 100644
diff --git a/Engine/Plugins/Experimental/ConsoleVariables/Source/ConsoleVariablesEditor/Private/Factories/ConsoleVariablesEditorFactory.h b/Engine/Plugins/Experimental/ConsoleVariables/Source/ConsoleVariablesEditor/Private/Factories/ConsoleVariablesEditorFactory.h
deleted file mode 100644
diff --git a/Engine/Plugins/Experimental/ConsoleVariables/Source/ConsoleVariablesEditor/Private/MultiUser/ConcertConsoleVariableSessionCustomization.cpp b/Engine/Plugins/Experimental/ConsoleVariables/Source/ConsoleVariablesEditor/Private/MultiUser/ConcertConsoleVariableSessionCustomization.cpp
deleted file mode 100644
diff --git a/Engine/Plugins/Experimental/ConsoleVariables/Source/ConsoleVariablesEditor/Private/MultiUser/ConcertConsoleVariableSessionCustomization.h b/Engine/Plugins/Experimental/ConsoleVariables/Source/ConsoleVariablesEditor/Private/MultiUser/ConcertConsoleVariableSessionCustomization.h
deleted file mode 100644
diff --git a/Engine/Plugins/Experimental/ConsoleVariables/Source/ConsoleVariablesEditor/Private/MultiUser/ConcertConsoleVariableSyncCustomization.h b/Engine/Plugins/Experimental/ConsoleVariables/Source/ConsoleVariablesEditor/Private/MultiUser/ConcertConsoleVariableSyncCustomization.h
deleted file mode 100644
diff --git a/Engine/Plugins/Experimental/ConsoleVariables/Source/ConsoleVariablesEditor/Private/MultiUser/ConsoleVariableSync.cpp b/Engine/Plugins/Experimental/ConsoleVariables/Source/ConsoleVariablesEditor/Private/MultiUser/ConsoleVariableSync.cpp
deleted file mode 100644
diff --git a/Engine/Plugins/Experimental/ConsoleVariables/Source/ConsoleVariablesEditor/Private/MultiUser/ConsoleVariableSync.h b/Engine/Plugins/Experimental/ConsoleVariables/Source/ConsoleVariablesEditor/Private/MultiUser/ConsoleVariableSync.h
deleted file mode 100644
diff --git a/Engine/Plugins/Experimental/ConsoleVariables/Source/ConsoleVariablesEditor/Private/MultiUser/ConsoleVariableSyncData.h b/Engine/Plugins/Experimental/ConsoleVariables/Source/ConsoleVariablesEditor/Private/MultiUser/ConsoleVariableSyncData.h
deleted file mode 100644
diff --git a/Engine/Plugins/Experimental/ConsoleVariables/Source/ConsoleVariablesEditor/Private/Views/List/ConsoleVariablesEditorList.cpp b/Engine/Plugins/Experimental/ConsoleVariables/Source/ConsoleVariablesEditor/Private/Views/List/ConsoleVariablesEditorList.cpp
deleted file mode 100644
diff --git a/Engine/Plugins/Experimental/ConsoleVariables/Source/ConsoleVariablesEditor/Private/Views/List/ConsoleVariablesEditorList.h b/Engine/Plugins/Experimental/ConsoleVariables/Source/ConsoleVariablesEditor/Private/Views/List/ConsoleVariablesEditorList.h
deleted file mode 100644
diff --git a/Engine/Plugins/Experimental/ConsoleVariables/Source/ConsoleVariablesEditor/Private/Views/List/ConsoleVariablesEditorListFilters/ConsoleVariablesEditorListFilter_ModifiedVariables.h b/Engine/Plugins/Experimental/ConsoleVariables/Source/ConsoleVariablesEditor/Private/Views/List/ConsoleVariablesEditorListFilters/ConsoleVariablesEditorListFilter_ModifiedVariables.h
deleted file mode 100644
diff --git a/Engine/Plugins/Experimental/ConsoleVariables/Source/ConsoleVariablesEditor/Private/Views/List/ConsoleVariablesEditorListFilters/ConsoleVariablesEditorListFilter_SourceText.h b/Engine/Plugins/Experimental/ConsoleVariables/Source/ConsoleVariablesEditor/Private/Views/List/ConsoleVariablesEditorListFilters/ConsoleVariablesEditorListFilter_SourceText.h
deleted file mode 100644
diff --git a/Engine/Plugins/Experimental/ConsoleVariables/Source/ConsoleVariablesEditor/Private/Views/List/ConsoleVariablesEditorListFilters/IConsoleVariablesEditorListFilter.h b/Engine/Plugins/Experimental/ConsoleVariables/Source/ConsoleVariablesEditor/Private/Views/List/ConsoleVariablesEditorListFilters/IConsoleVariablesEditorListFilter.h
deleted file mode 100644
diff --git a/Engine/Plugins/Experimental/ConsoleVariables/Source/ConsoleVariablesEditor/Private/Views/List/ConsoleVariablesEditorListRow.cpp b/Engine/Plugins/Experimental/ConsoleVariables/Source/ConsoleVariablesEditor/Private/Views/List/ConsoleVariablesEditorListRow.cpp
deleted file mode 100644
diff --git a/Engine/Plugins/Experimental/ConsoleVariables/Source/ConsoleVariablesEditor/Private/Views/List/ConsoleVariablesEditorListRow.h b/Engine/Plugins/Experimental/ConsoleVariables/Source/ConsoleVariablesEditor/Private/Views/List/ConsoleVariablesEditorListRow.h
deleted file mode 100644
diff --git a/Engine/Plugins/Experimental/ConsoleVariables/Source/ConsoleVariablesEditor/Private/Views/List/SConsoleVariablesEditorList.cpp b/Engine/Plugins/Experimental/ConsoleVariables/Source/ConsoleVariablesEditor/Private/Views/List/SConsoleVariablesEditorList.cpp
deleted file mode 100644
diff --git a/Engine/Plugins/Experimental/ConsoleVariables/Source/ConsoleVariablesEditor/Private/Views/List/SConsoleVariablesEditorList.h b/Engine/Plugins/Experimental/ConsoleVariables/Source/ConsoleVariablesEditor/Private/Views/List/SConsoleVariablesEditorList.h
deleted file mode 100644
diff --git a/Engine/Plugins/Experimental/ConsoleVariables/Source/ConsoleVariablesEditor/Private/Views/List/SConsoleVariablesEditorListRow.cpp b/Engine/Plugins/Experimental/ConsoleVariables/Source/ConsoleVariablesEditor/Private/Views/List/SConsoleVariablesEditorListRow.cpp
deleted file mode 100644
diff --git a/Engine/Plugins/Experimental/ConsoleVariables/Source/ConsoleVariablesEditor/Private/Views/List/SConsoleVariablesEditorListRow.h b/Engine/Plugins/Experimental/ConsoleVariables/Source/ConsoleVariablesEditor/Private/Views/List/SConsoleVariablesEditorListRow.h
deleted file mode 100644
diff --git a/Engine/Plugins/Experimental/ConsoleVariables/Source/ConsoleVariablesEditor/Private/Views/List/SConsoleVariablesEditorListValueInput.cpp b/Engine/Plugins/Experimental/ConsoleVariables/Source/ConsoleVariablesEditor/Private/Views/List/SConsoleVariablesEditorListValueInput.cpp
deleted file mode 100644
diff --git a/Engine/Plugins/Experimental/ConsoleVariables/Source/ConsoleVariablesEditor/Private/Views/List/SConsoleVariablesEditorListValueInput.h b/Engine/Plugins/Experimental/ConsoleVariables/Source/ConsoleVariablesEditor/Private/Views/List/SConsoleVariablesEditorListValueInput.h
deleted file mode 100644
diff --git a/Engine/Plugins/Experimental/ConsoleVariables/Source/ConsoleVariablesEditor/Private/Views/MainPanel/ConsoleVariablesEditorMainPanel.cpp b/Engine/Plugins/Experimental/ConsoleVariables/Source/ConsoleVariablesEditor/Private/Views/MainPanel/ConsoleVariablesEditorMainPanel.cpp
deleted file mode 100644
diff --git a/Engine/Plugins/Experimental/ConsoleVariables/Source/ConsoleVariablesEditor/Private/Views/MainPanel/ConsoleVariablesEditorMainPanel.h b/Engine/Plugins/Experimental/ConsoleVariables/Source/ConsoleVariablesEditor/Private/Views/MainPanel/ConsoleVariablesEditorMainPanel.h
deleted file mode 100644
diff --git a/Engine/Plugins/Experimental/ConsoleVariables/Source/ConsoleVariablesEditor/Private/Views/MainPanel/SConsoleVariablesEditorCustomConsoleInputBox.cpp b/Engine/Plugins/Experimental/ConsoleVariables/Source/ConsoleVariablesEditor/Private/Views/MainPanel/SConsoleVariablesEditorCustomConsoleInputBox.cpp
deleted file mode 100644
diff --git a/Engine/Plugins/Experimental/ConsoleVariables/Source/ConsoleVariablesEditor/Private/Views/MainPanel/SConsoleVariablesEditorCustomConsoleInputBox.h b/Engine/Plugins/Experimental/ConsoleVariables/Source/ConsoleVariablesEditor/Private/Views/MainPanel/SConsoleVariablesEditorCustomConsoleInputBox.h
deleted file mode 100644
diff --git a/Engine/Plugins/Experimental/ConsoleVariables/Source/ConsoleVariablesEditor/Private/Views/MainPanel/SConsoleVariablesEditorMainPanel.cpp b/Engine/Plugins/Experimental/ConsoleVariables/Source/ConsoleVariablesEditor/Private/Views/MainPanel/SConsoleVariablesEditorMainPanel.cpp
deleted file mode 100644
diff --git a/Engine/Plugins/Experimental/ConsoleVariables/Source/ConsoleVariablesEditor/Private/Views/MainPanel/SConsoleVariablesEditorMainPanel.h b/Engine/Plugins/Experimental/ConsoleVariables/Source/ConsoleVariablesEditor/Private/Views/MainPanel/SConsoleVariablesEditorMainPanel.h
deleted file mode 100644
diff --git a/Engine/Plugins/Experimental/ConsoleVariables/Source/ConsoleVariablesEditor/Private/Views/Widgets/SConsoleVariablesEditorGlobalSearchToggle.cpp b/Engine/Plugins/Experimental/ConsoleVariables/Source/ConsoleVariablesEditor/Private/Views/Widgets/SConsoleVariablesEditorGlobalSearchToggle.cpp
deleted file mode 100644
diff --git a/Engine/Plugins/Experimental/ConsoleVariables/Source/ConsoleVariablesEditor/Private/Views/Widgets/SConsoleVariablesEditorGlobalSearchToggle.h b/Engine/Plugins/Experimental/ConsoleVariables/Source/ConsoleVariablesEditor/Private/Views/Widgets/SConsoleVariablesEditorGlobalSearchToggle.h
deleted file mode 100644
diff --git a/Engine/Plugins/Experimental/ConsoleVariables/Source/ConsoleVariablesEditor/Private/Views/Widgets/SConsoleVariablesEditorTooltipWidget.h b/Engine/Plugins/Experimental/ConsoleVariables/Source/ConsoleVariablesEditor/Private/Views/Widgets/SConsoleVariablesEditorTooltipWidget.h
deleted file mode 100644
diff --git a/Engine/Plugins/Experimental/ConsoleVariables/Source/ConsoleVariablesEditor/Public/ConsoleVariablesAsset.h b/Engine/Plugins/Experimental/ConsoleVariables/Source/ConsoleVariablesEditor/Public/ConsoleVariablesAsset.h
deleted file mode 100644
diff --git a/Engine/Plugins/Experimental/ConsoleVariables/Source/ConsoleVariablesEditor/Public/ConsoleVariablesEditorFunctionLibrary.h b/Engine/Plugins/Experimental/ConsoleVariables/Source/ConsoleVariablesEditor/Public/ConsoleVariablesEditorFunctionLibrary.h
deleted file mode 100644
diff --git a/Engine/Plugins/Experimental/ConsoleVariables/Source/ConsoleVariablesEditor/Public/ConsoleVariablesEditorLog.h b/Engine/Plugins/Experimental/ConsoleVariables/Source/ConsoleVariablesEditor/Public/ConsoleVariablesEditorLog.h
deleted file mode 100644
diff --git a/Engine/Plugins/Experimental/ConsoleVariables/Source/ConsoleVariablesEditor/Public/ConsoleVariablesEditorProjectSettings.h b/Engine/Plugins/Experimental/ConsoleVariables/Source/ConsoleVariablesEditor/Public/ConsoleVariablesEditorProjectSettings.h
deleted file mode 100644
