Console Variables Editor:

-Have 'no results' copy be the same when filtering preset or global search

#[fyi] Jason.Walter

diff --git a/Engine/Plugins/Editor/ConsoleVariablesEditor/Source/ConsoleVariablesEditor/Private/Views/List/SConsoleVariablesEditorList.cpp b/Engine/Plugins/Editor/ConsoleVariablesEditor/Source/ConsoleVariablesEditor/Private/Views/List/SConsoleVariablesEditorList.cpp
--- a/Engine/Plugins/Editor/ConsoleVariablesEditor/Source/ConsoleVariablesEditor/Private/Views/List/SConsoleVariablesEditorList.cpp
+++ b/Engine/Plugins/Editor/ConsoleVariablesEditor/Source/ConsoleVariablesEditor/Private/Views/List/SConsoleVariablesEditorList.cpp
@@ -230,17 +230,17 @@ void SConsoleVariablesEditorList::Construct(const FArguments& InArgs, TSharedRef
 				.Justification(ETextJustify::Center)
 				.Text_Lambda([this]()
 				{
-					if (const TSharedPtr<FConsoleVariablesEditorList> ListModel = ListModelPtr.Pin())
+					if (VisibleTreeViewObjects.Num() == TreeViewRootObjects.Num())
 					{
-						if (ListModel->GetListMode() == FConsoleVariablesEditorList::EConsoleVariablesEditorListMode::GlobalSearch)
-						{
-							return LOCTEXT("ConsoleVariablesEditorList_NoList",
-								"No matching console variables found in Unreal Engine.\n\nCheck your search criteria.");
-						}
+						// Global Search Empty List (without filter)
+                    	return LOCTEXT("ConsoleVariablesEditorList_EmptyListGlobalSearch",
+                    		"No matching console variables found in Unreal Engine.\n\nCheck your search criteria.");
+						
 					}
-					
-					return LOCTEXT("ConsoleVariablesEditorList_NoList",
-					"No matching console variables in your list.\n\nCheck your filter or <RichTextBlock.Bold>Search All</> console variables instead.");
+
+					// Empty List (with filter)
+                    return LOCTEXT("ConsoleVariablesEditorList_EmptyListWithFilter",
+                    	"No matching console variables in your list.\n\nCheck your filter or <RichTextBlock.Bold>Search All</> console variables instead.");
 				})
 			]
 		]
@@ -995,7 +995,7 @@ void SConsoleVariablesEditorList::SetAllGroupsCollapsed()
 
 void SConsoleVariablesEditorList::OnListViewSearchTextChanged(const FText& Text)
 {
-	ExecuteListViewSearchOnAllRows(Text.ToString(), false);
+	ExecuteListViewSearchOnAllRows(Text.ToString(), true);
 }
 
 void SConsoleVariablesEditorList::CacheCurrentListItemData()
