Console Variables Editor:

Fixes for:
-Search form is too tall (~4px on a 1920x1080 monitor), ensure it is the same height as the search form in the Outliner.
-"Search All" button is too tall (~5px on a 1920x1080 monitor), ensure it is the same height as the "+ Console Variable" button.
-The checkbox in the column header does not align with the checkboxes in the rows (header is slightly to the right).
-Improvements to No Results copy
-Add more default fill width to Sort Column
-Crash Creating Preset in Content Browser

diff --git a/Engine/Plugins/Editor/ConsoleVariablesEditor/Source/ConsoleVariablesEditor/Private/Factories/ConsoleVariablesEditorFactory.cpp b/Engine/Plugins/Editor/ConsoleVariablesEditor/Source/ConsoleVariablesEditor/Private/Factories/ConsoleVariablesEditorFactory.cpp
--- a/Engine/Plugins/Editor/ConsoleVariablesEditor/Source/ConsoleVariablesEditor/Private/Factories/ConsoleVariablesEditorFactory.cpp
+++ b/Engine/Plugins/Editor/ConsoleVariablesEditor/Source/ConsoleVariablesEditor/Private/Factories/ConsoleVariablesEditorFactory.cpp
@@ -21,7 +21,7 @@ UConsoleVariablesEditorFactory::UConsoleVariablesEditorFactory()
 
 UObject* UConsoleVariablesEditorFactory::FactoryCreateNew(UClass* InClass, UObject* InParent, FName InName, EObjectFlags Flags, UObject* Context, FFeedbackContext* Warn)
 {
-	return NewObject<UConsoleVariablesEditorFactory>(InParent, InClass, InName, Flags);
+	return NewObject<UConsoleVariablesAsset>(InParent, InClass, InName, Flags);
 };
 
 
diff --git a/Engine/Plugins/Editor/ConsoleVariablesEditor/Source/ConsoleVariablesEditor/Private/Views/List/SConsoleVariablesEditorList.cpp b/Engine/Plugins/Editor/ConsoleVariablesEditor/Source/ConsoleVariablesEditor/Private/Views/List/SConsoleVariablesEditorList.cpp
--- a/Engine/Plugins/Editor/ConsoleVariablesEditor/Source/ConsoleVariablesEditor/Private/Views/List/SConsoleVariablesEditorList.cpp
+++ b/Engine/Plugins/Editor/ConsoleVariablesEditor/Source/ConsoleVariablesEditor/Private/Views/List/SConsoleVariablesEditorList.cpp
@@ -60,6 +60,7 @@ void SConsoleVariablesEditorList::Construct(const FArguments& InArgs, TSharedRef
 
 			+SHorizontalBox::Slot()
 			.HAlign(HAlign_Fill)
+			.VAlign(VAlign_Center)
 			.Padding(0.f, 1.f, 0.f, 1.f)
 			[
 				SAssignNew(ListSearchBoxPtr, SSearchBox)
@@ -76,12 +77,13 @@ void SConsoleVariablesEditorList::Construct(const FArguments& InArgs, TSharedRef
 			.Padding(FMargin(8.f, 0.f, 0.f, 0.f))
 			[
 				SNew(SButton)
-				.ContentPadding(4.f)
+				.ContentPadding(FMargin(4, 1.5f))
 				.VAlign(VAlign_Center)
 				.OnClicked_Lambda([this]()
 				{
 					const FString SearchString = GetSearchStringFromSearchInputField().ToLower().TrimStartAndEnd();
 					ListSearchBoxPtr->SetText(FText::GetEmpty());
+					VisibleTreeViewObjects.Empty();
 					return TryEnterGlobalSearch(SearchString);
 				})
 				.ToolTipText(LOCTEXT("OpenInGlobalSearchButtonTooltip", "Search All Console Variables"))
@@ -760,7 +762,7 @@ TSharedPtr<SHeaderRow> SConsoleVariablesEditorList::GenerateHeaderRow()
 			.DefaultLabel(FText::FromString("#"))
 			.ToolTipText(LOCTEXT("ClickToSort","Click to sort"))
 			.HAlignHeader(EHorizontalAlignment::HAlign_Center)
-			.FillWidth(0.3f)
+			.FillWidth(0.4f)
 			.ShouldGenerateWidget(true)
 			.SortMode_Raw(this, &SConsoleVariablesEditorList::GetSortModeForColumn, CustomSortOrderColumnName)
 			.OnSort_Raw(this, &SConsoleVariablesEditorList::OnSortColumnCalled)
diff --git a/Engine/Plugins/Editor/ConsoleVariablesEditor/Source/ConsoleVariablesEditor/Private/Views/List/SConsoleVariablesEditorListRow.cpp b/Engine/Plugins/Editor/ConsoleVariablesEditor/Source/ConsoleVariablesEditor/Private/Views/List/SConsoleVariablesEditorListRow.cpp
--- a/Engine/Plugins/Editor/ConsoleVariablesEditor/Source/ConsoleVariablesEditor/Private/Views/List/SConsoleVariablesEditorListRow.cpp
+++ b/Engine/Plugins/Editor/ConsoleVariablesEditor/Source/ConsoleVariablesEditor/Private/Views/List/SConsoleVariablesEditorListRow.cpp
@@ -352,14 +352,16 @@ TSharedRef<SWidget> SConsoleVariablesEditorListRow::GenerateCells(const FName& I
 		return SNew(SBox)
 				.Visibility(EVisibility::SelfHitTestInvisible)
 				.HAlign(HAlign_Center)
+				.Padding(FMargin(1,0,0,0))
 		[
 			SNew(SCheckBox)
-					.IsChecked_Raw(this, &SConsoleVariablesEditorListRow::GetCheckboxState)
-					.OnCheckStateChanged_Raw(this, &SConsoleVariablesEditorListRow::OnCheckboxStateChange)
-					.Visibility(Item.Pin()->GetCommandInfo().Pin()->ObjectType ==
-					            FConsoleVariablesEditorCommandInfo::EConsoleObjectType::Variable
-						            ? EVisibility::Visible
-						            : EVisibility::Collapsed)
+			.HAlign(HAlign_Center)
+			.IsChecked_Raw(this, &SConsoleVariablesEditorListRow::GetCheckboxState)
+			.OnCheckStateChanged_Raw(this, &SConsoleVariablesEditorListRow::OnCheckboxStateChange)
+			.Visibility(Item.Pin()->GetCommandInfo().Pin()->ObjectType ==
+			            FConsoleVariablesEditorCommandInfo::EConsoleObjectType::Variable
+				            ? EVisibility::Visible
+				            : EVisibility::Collapsed)
 		];
 	}
 	if (InColumnName.IsEqual(SConsoleVariablesEditorList::VariableNameColumnName))
