Console Variables Editor:

Fixes for:
-Duplicated keys for localization - UE - Virtual Production

diff --git a/Engine/Plugins/Editor/ConsoleVariablesEditor/Source/ConsoleVariablesEditor/Private/ConsoleVariablesEditorModule.cpp b/Engine/Plugins/Editor/ConsoleVariablesEditor/Source/ConsoleVariablesEditor/Private/ConsoleVariablesEditorModule.cpp
--- a/Engine/Plugins/Editor/ConsoleVariablesEditor/Source/ConsoleVariablesEditor/Private/ConsoleVariablesEditorModule.cpp
+++ b/Engine/Plugins/Editor/ConsoleVariablesEditor/Source/ConsoleVariablesEditor/Private/ConsoleVariablesEditorModule.cpp
@@ -326,9 +326,8 @@ void FConsoleVariablesEditorModule::RegisterProjectSettings() const
 		// User Project Settings
 		const TSharedPtr<ISettingsSection> ProjectSettingsSectionPtr = SettingsModule.RegisterSettings(
 			"Project", "Plugins", "Console Variables Editor",
-			NSLOCTEXT("ConsoleVariables", "ConsoleVariablesSettingsCategoryDisplayName", "Console Variables Editor"),
-			NSLOCTEXT("ConsoleVariables", "ConsoleVariablesSettingsDescription",
-			          "Configure the Console Variables Editor user settings"),
+			LOCTEXT("ConsoleVariablesSettingsCategoryDisplayName", "Console Variables Editor"),
+			LOCTEXT("ConsoleVariablesSettingsDescription", "Configure the Console Variables Editor user settings"),
 			GetMutableDefault<UConsoleVariablesEditorProjectSettings>());
 	}
 }
diff --git a/Engine/Plugins/Editor/ConsoleVariablesEditor/Source/ConsoleVariablesEditor/Private/Views/List/ConsoleVariablesEditorListFilters/ConsoleVariablesEditorListFilter_ModifiedVariables.h b/Engine/Plugins/Editor/ConsoleVariablesEditor/Source/ConsoleVariablesEditor/Private/Views/List/ConsoleVariablesEditorListFilters/ConsoleVariablesEditorListFilter_ModifiedVariables.h
--- a/Engine/Plugins/Editor/ConsoleVariablesEditor/Source/ConsoleVariablesEditor/Private/Views/List/ConsoleVariablesEditorListFilters/ConsoleVariablesEditorListFilter_ModifiedVariables.h
+++ b/Engine/Plugins/Editor/ConsoleVariablesEditor/Source/ConsoleVariablesEditor/Private/Views/List/ConsoleVariablesEditorListFilters/ConsoleVariablesEditorListFilter_ModifiedVariables.h
@@ -29,7 +29,7 @@ public:
 
 	virtual FText GetFilterButtonToolTip() override
 	{
-		return LOCTEXT("ShowSourceTextFilterTooltipFormat", "Show only rows that have a current value that differs from the preset value.");
+		return LOCTEXT("ShowOnlyModifiedVariablesFilterTooltipFormat", "Show only rows that have a current value that differs from the preset value.");
 	}
 
 	virtual bool DoesItemPassFilter(const FConsoleVariablesEditorListRowPtr& InItem) override
