Console Variables Editor:

-Logging verbosity pass
-Fix for MU implementation sending variables back and forth in a feedback loop
-Fix for Global Search does not crawl help text or source

diff --git a/Engine/Plugins/Editor/ConsoleVariablesEditor/Source/ConsoleVariablesEditor/Private/ConsoleVariablesAsset.cpp b/Engine/Plugins/Editor/ConsoleVariablesEditor/Source/ConsoleVariablesEditor/Private/ConsoleVariablesAsset.cpp
--- a/Engine/Plugins/Editor/ConsoleVariablesEditor/Source/ConsoleVariablesEditor/Private/ConsoleVariablesAsset.cpp
+++ b/Engine/Plugins/Editor/ConsoleVariablesEditor/Source/ConsoleVariablesEditor/Private/ConsoleVariablesAsset.cpp
@@ -41,7 +41,7 @@ void UConsoleVariablesAsset::AddOrSetConsoleObjectSavedData(const FConsoleVariab
 		
 		if (AsVariable && AsVariable->TestFlags(ECVF_RenderThreadSafe))
 		{
-			UE_LOG(LogConsoleVariablesEditor, Warning,
+			UE_LOG(LogConsoleVariablesEditor, Verbose,
 				TEXT("The console variable named %s is flagged as ECVF_RenderThreadSafe. The value on the render thread will lag behind the value on the main thread by one frame if r.OneFrameThreadLag is 1."),
 				*InData.CommandName
 			);
diff --git a/Engine/Plugins/Editor/ConsoleVariablesEditor/Source/ConsoleVariablesEditor/Private/ConsoleVariablesEditorModule.cpp b/Engine/Plugins/Editor/ConsoleVariablesEditor/Source/ConsoleVariablesEditor/Private/ConsoleVariablesEditorModule.cpp
--- a/Engine/Plugins/Editor/ConsoleVariablesEditor/Source/ConsoleVariablesEditor/Private/ConsoleVariablesEditorModule.cpp
+++ b/Engine/Plugins/Editor/ConsoleVariablesEditor/Source/ConsoleVariablesEditor/Private/ConsoleVariablesEditorModule.cpp
@@ -162,6 +162,7 @@ TArray<TWeakPtr<FConsoleVariablesEditorCommandInfo>> FConsoleVariablesEditorModu
 	
 	for (const TSharedPtr<FConsoleVariablesEditorCommandInfo>& CommandInfo : ConsoleObjectsMasterReference)
 	{
+		FString CommandSearchableText = CommandInfo->Command + " " + CommandInfo->GetHelpText() + " " + CommandInfo->GetSourceAsText().ToString();
 		// Match any
 		for (const FString& Token : InTokens)
 		{
@@ -172,14 +173,14 @@ TArray<TWeakPtr<FConsoleVariablesEditorCommandInfo>> FConsoleVariablesEditorModu
 			TArray<FString> OutSpacedArray;
 			if (Token.Contains(SpaceDelimiter) && Token.ParseIntoArray(OutSpacedArray, *SpaceDelimiter, true) > 1)
 			{
-				bMatchFound = Algo::AllOf(OutSpacedArray, [&CommandInfo, InSearchCase](const FString& Comparator)
+				bMatchFound = Algo::AllOf(OutSpacedArray, [&CommandSearchableText, InSearchCase](const FString& Comparator)
 				{
-					return CommandInfo->Command.Contains(Comparator, InSearchCase);
+					return CommandSearchableText.Contains(Comparator, InSearchCase);
 				});
 			}
 			else
 			{
-				bMatchFound = CommandInfo->Command.Contains(Token, InSearchCase);
+				bMatchFound = CommandSearchableText.Contains(Token, InSearchCase);
 			}
 
 			if (bMatchFound)
@@ -263,7 +264,9 @@ void FConsoleVariablesEditorModule::SendMultiUserConsoleVariableChange(const FSt
 
 void FConsoleVariablesEditorModule::OnRemoteCvarChanged(const FString InName, const FString InValue)
 {
-	UE_LOG(LogConsoleVariablesEditor, Display, TEXT("Remote set console variable %s = %s"), *InName, *InValue);
+	UE_LOG(LogConsoleVariablesEditor, VeryVerbose, TEXT("Remote set console variable %s = %s"), *InName, *InValue);
+
+	CommandsRecentlyReceivedFromMultiUser.Add(InName, InValue);
 
 	if (GetDefault<UConcertCVarSynchronization>()->bSyncCVarTransactions)
 	{
@@ -343,7 +346,7 @@ void FConsoleVariablesEditorModule::OnConsoleVariableChanged(IConsoleVariable* C
 		const FString& Key = PinnedCommand->Command;
 		
 		FConsoleVariablesEditorAssetSaveData FoundData;
-		const bool bIsVariableCurrentlyTracked = EditingPresetAsset->FindSavedDataByCommandString(Key, FoundData);
+		bool bIsVariableCurrentlyTracked = EditingPresetAsset->FindSavedDataByCommandString(Key, FoundData);
 
 		const UConsoleVariablesEditorProjectSettings* Settings = GetDefault<UConsoleVariablesEditorProjectSettings>();
 		check(Settings);
@@ -367,18 +370,35 @@ void FConsoleVariablesEditorModule::OnConsoleVariableChanged(IConsoleVariable* C
 						ChangedVariable->GetString() : "",
 						true
 					);
-				}
 
-				SendMultiUserConsoleVariableChange(Key, ChangedVariable->GetString());
+					bIsVariableCurrentlyTracked = true;
+				}
 			}
 		}
-		else // If it's already being tracked, refreshed the list to update show filters and other possibly stale elements
+
+		// If the variable is already tracked or was just added to tracking, run the following code
+		if (bIsVariableCurrentlyTracked)
 		{
 			if (MainPanel.IsValid())
 			{
 				MainPanel->RefreshList();
 			}
 			
+			/**
+			 * Here we check the map of recently received variables from other nodes
+			 * If the command is in the map and the value is similar, we won't send the value to other nodes
+			 * because we can assume that this value came from another node.
+			 * This prevents a feedback loop.
+			 */
+			if (const FString* MatchedValue = CommandsRecentlyReceivedFromMultiUser.Find(Key))
+			{
+				if (MatchedValue->Equals(ChangedVariable->GetString()))
+				{
+					CommandsRecentlyReceivedFromMultiUser.Remove(Key);
+					return;
+				}
+			}
+				
 			SendMultiUserConsoleVariableChange(Key, ChangedVariable->GetString());
 		}
 	}
diff --git a/Engine/Plugins/Editor/ConsoleVariablesEditor/Source/ConsoleVariablesEditor/Private/ConsoleVariablesEditorModule.h b/Engine/Plugins/Editor/ConsoleVariablesEditor/Source/ConsoleVariablesEditor/Private/ConsoleVariablesEditorModule.h
--- a/Engine/Plugins/Editor/ConsoleVariablesEditor/Source/ConsoleVariablesEditor/Private/ConsoleVariablesEditorModule.h
+++ b/Engine/Plugins/Editor/ConsoleVariablesEditor/Source/ConsoleVariablesEditor/Private/ConsoleVariablesEditorModule.h
@@ -110,7 +110,7 @@ private:
 	
 	static const FName ConsoleVariablesToolkitPanelTabId;
 
-	/* Lives for as long as the module is loaded. */
+	/** Lives for as long as the module is loaded. */
 	TSharedPtr<FConsoleVariablesEditorMainPanel> MainPanel;
 
 	/** Transient preset that's being edited so we don't affect the reference asset unless we save it */
@@ -120,4 +120,12 @@ private:
 
 	/** All tracked variables and their default, startup, and current values */
 	TArray<TSharedPtr<FConsoleVariablesEditorCommandInfo>> ConsoleObjectsMasterReference;
+
+	/**
+	 * A list of CommandName/ValueAsString pairs of commands received from MU.
+	 * When this node receives a command and value from another node, it will be added to this list.
+	 * This list will be checked against detected local cvar changes and,
+	 * if the command and value match an item in this map, the change will not be propagated to other nodes.
+	 */
+	TMap<FString, FString> CommandsRecentlyReceivedFromMultiUser;
 };
diff --git a/Engine/Plugins/Editor/ConsoleVariablesEditor/Source/ConsoleVariablesEditor/Private/MultiUser/ConsoleVariableSync.cpp b/Engine/Plugins/Editor/ConsoleVariablesEditor/Source/ConsoleVariablesEditor/Private/MultiUser/ConsoleVariableSync.cpp
--- a/Engine/Plugins/Editor/ConsoleVariablesEditor/Source/ConsoleVariablesEditor/Private/MultiUser/ConsoleVariableSync.cpp
+++ b/Engine/Plugins/Editor/ConsoleVariablesEditor/Source/ConsoleVariablesEditor/Private/MultiUser/ConsoleVariableSync.cpp
@@ -124,7 +124,7 @@ struct FManagerImpl
 
 	void ConnectToSession(IConcertClientSession& InSession)
 	{
-		UE_LOG(LogConsoleVariablesEditor, Display, TEXT("Multi-user Console Variable Editor listening to session: %s"), *InSession.GetSessionInfo().SessionName);
+		UE_LOG(LogConsoleVariablesEditor, VeryVerbose, TEXT("Multi-user Console Variable Editor listening to session: %s"), *InSession.GetSessionInfo().SessionName);
 
 		ClientChangeDelegate = InSession.OnSessionClientChanged().AddRaw(this, &FManagerImpl::OnSessionClientChanged);
 
@@ -229,7 +229,7 @@ struct FManagerImpl
 		check(WeakSession.IsValid());
 		{
 			TSharedPtr<IConcertClientSession> Session = WeakSession.Pin();
-			UE_LOG(LogConsoleVariablesEditor, Display, TEXT("Multi-user Console Variable Editor disconnecting from Session: %s"), *Session->GetSessionInfo().SessionName);
+			UE_LOG(LogConsoleVariablesEditor, VeryVerbose, TEXT("Multi-user Console Variable Editor disconnecting from Session: %s"), *Session->GetSessionInfo().SessionName);
 			Session->OnSessionClientChanged().Remove(ClientChangeDelegate);
 		}
 	}
diff --git a/Engine/Plugins/Editor/ConsoleVariablesEditor/Source/ConsoleVariablesEditor/Private/Views/List/SConsoleVariablesEditorList.cpp b/Engine/Plugins/Editor/ConsoleVariablesEditor/Source/ConsoleVariablesEditor/Private/Views/List/SConsoleVariablesEditorList.cpp
--- a/Engine/Plugins/Editor/ConsoleVariablesEditor/Source/ConsoleVariablesEditor/Private/Views/List/SConsoleVariablesEditorList.cpp
+++ b/Engine/Plugins/Editor/ConsoleVariablesEditor/Source/ConsoleVariablesEditor/Private/Views/List/SConsoleVariablesEditorList.cpp
@@ -275,7 +275,7 @@ FReply SConsoleVariablesEditorList::TryEnterGlobalSearch(const FString& SearchSt
 	// Can't enter global search if there are no active global searches or a search string from which to parse new searches
 	if (SearchString.IsEmpty() && CurrentGlobalSearches.Num() == 0)
 	{
-		UE_LOG(LogConsoleVariablesEditor, Log,
+		UE_LOG(LogConsoleVariablesEditor, Verbose,
 			TEXT("%hs: Global search request is empty. Exiting Global Search."),
 			__FUNCTION__, *SearchString);
 
diff --git a/Engine/Plugins/Editor/ConsoleVariablesEditor/Source/ConsoleVariablesEditor/Private/Views/MainPanel/ConsoleVariablesEditorMainPanel.cpp b/Engine/Plugins/Editor/ConsoleVariablesEditor/Source/ConsoleVariablesEditor/Private/Views/MainPanel/ConsoleVariablesEditorMainPanel.cpp
--- a/Engine/Plugins/Editor/ConsoleVariablesEditor/Source/ConsoleVariablesEditor/Private/Views/MainPanel/ConsoleVariablesEditorMainPanel.cpp
+++ b/Engine/Plugins/Editor/ConsoleVariablesEditor/Source/ConsoleVariablesEditor/Private/Views/MainPanel/ConsoleVariablesEditorMainPanel.cpp
@@ -280,10 +280,10 @@ void FConsoleVariablesEditorMainPanel::OnConnectionChanged(EConcertConnectionSta
 	switch (Status)
 	{
 		case EConcertConnectionStatus::Connected:
-			UE_LOG(LogConsoleVariablesEditor, Display, TEXT("Multi-user has connected to a session."));
+			UE_LOG(LogConsoleVariablesEditor, VeryVerbose, TEXT("Multi-user has connected to a session."));
 			break;
 		case EConcertConnectionStatus::Disconnected:
-			UE_LOG(LogConsoleVariablesEditor, Display, TEXT("Multi-user has disconnected from session."));
+			UE_LOG(LogConsoleVariablesEditor, VeryVerbose, TEXT("Multi-user has disconnected from session."));
 			break;
 		default:
 			break;
