Console Variables Editor:

-Plugin icon update
-Adjust no list text for when there are no variables with no filter in preset mode
-Treat same commands with separate parameters before value as separate commands in list
-Fix caching for value commands which were being reset when returning from global search
-Fix: Engine cVars should become "Session" cVars when modified
-Fix: Show Option "Show Only Current Preset" not working
-Add Show Option "Only Variables"
-Add Show Option "Only Commands"
-Update Show Options to use EConsoleVariableFlags instead of text matching

diff --git a/Engine/Plugins/Editor/ConsoleVariablesEditor/Source/ConsoleVariablesEditor/Private/ConsoleVariablesEditorCommandInfo.h b/Engine/Plugins/Editor/ConsoleVariablesEditor/Source/ConsoleVariablesEditor/Private/ConsoleVariablesEditorCommandInfo.h
--- a/Engine/Plugins/Editor/ConsoleVariablesEditor/Source/ConsoleVariablesEditor/Private/ConsoleVariablesEditorCommandInfo.h
+++ b/Engine/Plugins/Editor/ConsoleVariablesEditor/Source/ConsoleVariablesEditor/Private/ConsoleVariablesEditorCommandInfo.h
@@ -62,10 +62,22 @@ struct FConsoleVariablesEditorCommandInfo
 		return Command.Equals(Comparator.Command);
 	}
 
-	void ExecuteCommand(const FString& NewValueAsString) const
+	void SetIfChangedInCurrentPreset(const bool bNewSetting)
 	{
-		GEngine->Exec(GetCurrentWorld(),
-			*FString::Printf(TEXT("%s %s"), *Command, *NewValueAsString).TrimStartAndEnd());
+		bChangedInCurrentPreset = bNewSetting;
+	}
+
+	void ExecuteCommand(const FString& NewValueAsString)
+	{
+		if (IConsoleVariable* AsVariable = GetConsoleVariablePtr())
+		{
+			AsVariable->Set(*NewValueAsString, GetSource());
+		}
+		else
+		{
+			GEngine->Exec(GetCurrentWorld(),
+				*FString::Printf(TEXT("%s %s"), *Command, *NewValueAsString).TrimStartAndEnd());
+		}
 	}
 
 	/** Get a reference to the cached console object. May return nullptr if unregistered. */
@@ -76,7 +88,16 @@ struct FConsoleVariablesEditorCommandInfo
 		if (!ConsoleObjectPtr ||
 			(FDateTime::UtcNow() - TimeOfLastConsoleObjectRefresh).GetTotalSeconds() > ConsoleObjectRefreshThreshold)
 		{
-			ConsoleObjectPtr = IConsoleManager::Get().FindConsoleObject(*Command);
+			FString CommandKey = Command; 
+
+			// Remove additional params, if they exist
+			const int32 IndexOfSpace = CommandKey.Find(" ");
+			if (IndexOfSpace != INDEX_NONE)
+			{
+				CommandKey = CommandKey.Left(IndexOfSpace).TrimStartAndEnd();
+			}
+			
+			ConsoleObjectPtr = IConsoleManager::Get().FindConsoleObject(*CommandKey);
 			TimeOfLastConsoleObjectRefresh = FDateTime::UtcNow();
 		}
 
@@ -161,6 +182,12 @@ struct FConsoleVariablesEditorCommandInfo
 
 	void SetSourceFlag(const EConsoleVariableFlags InSource)
 	{
+		if (IConsoleVariable* AsVariable = GetConsoleVariablePtr())
+		{
+			AsVariable->Set(*AsVariable->GetString(), StartupSource);
+			return;
+		}
+		
 		const uint32 OldPri = (uint32)GetSource();
 		const uint32 NewPri = (uint32)InSource;
 
@@ -178,6 +205,17 @@ struct FConsoleVariablesEditorCommandInfo
 
 	FText GetSourceAsText()
 	{
+		// Non-variables don't really have a source
+		if (ObjectType != EConsoleObjectType::Variable)
+		{
+			return LOCTEXT("Source_IsNotConsoleVariableButConsoleCommand", "Command");
+		}
+
+		if (bChangedInCurrentPreset)
+		{
+			return LOCTEXT("Source_SetByCurrentPreset", "Session");
+		}
+		
 		return ConvertConsoleVariableSetByFlagToText(GetSource());
 	}
 
@@ -239,6 +277,9 @@ struct FConsoleVariablesEditorCommandInfo
 	/** The source of this variable's (if Variable object type) last setting as recorded when the plugin was loaded. */
 	EConsoleVariableFlags StartupSource = ECVF_Default;
 
+	/** If the variable was last changed by the current preset */
+	bool bChangedInCurrentPreset = false;
+
 	/** When variables change, this callback is executed. */
 	FDelegateHandle OnVariableChangedCallbackHandle;
 
@@ -249,17 +290,17 @@ struct FConsoleVariablesEditorCommandInfo
 	/** A mapping of SetBy console variable flags to information like the associated display text. */
 	static const inline TArray<FStaticConsoleVariableFlagInfo> SupportedFlags =
 	{
-		{ EConsoleVariableFlags::ECVF_SetByConstructor, LOCTEXT("SetByConstructor", "Constructor") },
-		{ EConsoleVariableFlags::ECVF_SetByScalability, LOCTEXT("SetByScalability", "Scalability") },
-		{ EConsoleVariableFlags::ECVF_SetByGameSetting, LOCTEXT("SetByGameSetting", "Game Setting") },
-		{ EConsoleVariableFlags::ECVF_SetByProjectSetting, LOCTEXT("SetByProjectSetting", "Project Setting") },
-		{ EConsoleVariableFlags::ECVF_SetBySystemSettingsIni, LOCTEXT("SetBySystemSettingsIni", "System Settings ini") },
-		{ EConsoleVariableFlags::ECVF_SetByDeviceProfile, LOCTEXT("SetByDeviceProfile", "Device Profile") },
-		{ EConsoleVariableFlags::ECVF_SetByGameOverride, LOCTEXT("SetByGameOverride", "Game Override") },
-		{ EConsoleVariableFlags::ECVF_SetByConsoleVariablesIni, LOCTEXT("SetByConsoleVariablesIni", "Console Variables ini") },
-		{ EConsoleVariableFlags::ECVF_SetByCommandline, LOCTEXT("SetByCommandline", "Command line") },
-		{ EConsoleVariableFlags::ECVF_SetByCode, LOCTEXT("SetByCode", "Code") },
-		{ EConsoleVariableFlags::ECVF_SetByConsole, LOCTEXT("SetByConsole", "Console") }
+		{ EConsoleVariableFlags::ECVF_SetByConstructor, LOCTEXT("Source_SetByConstructor", "Constructor") },
+		{ EConsoleVariableFlags::ECVF_SetByScalability, LOCTEXT("Source_SetByScalability", "Scalability") },
+		{ EConsoleVariableFlags::ECVF_SetByGameSetting, LOCTEXT("Source_SetByGameSetting", "Game Setting") },
+		{ EConsoleVariableFlags::ECVF_SetByProjectSetting, LOCTEXT("Source_SetByProjectSetting", "Project Setting") },
+		{ EConsoleVariableFlags::ECVF_SetBySystemSettingsIni, LOCTEXT("Source_SetBySystemSettingsIni", "System Settings ini") },
+		{ EConsoleVariableFlags::ECVF_SetByDeviceProfile, LOCTEXT("Source_SetByDeviceProfile", "Device Profile") },
+		{ EConsoleVariableFlags::ECVF_SetByGameOverride, LOCTEXT("Source_SetByGameOverride", "Game Override") },
+		{ EConsoleVariableFlags::ECVF_SetByConsoleVariablesIni, LOCTEXT("Source_SetByConsoleVariablesIni", "Console Variables ini") },
+		{ EConsoleVariableFlags::ECVF_SetByCommandline, LOCTEXT("Source_SetByCommandline", "Command line") },
+		{ EConsoleVariableFlags::ECVF_SetByCode, LOCTEXT("Source_SetByCode", "Code") },
+		{ EConsoleVariableFlags::ECVF_SetByConsole, LOCTEXT("Source_SetByConsole", "Console") }
 	};
 };
 
diff --git a/Engine/Plugins/Editor/ConsoleVariablesEditor/Source/ConsoleVariablesEditor/Private/Views/List/ConsoleVariablesEditorList.cpp b/Engine/Plugins/Editor/ConsoleVariablesEditor/Source/ConsoleVariablesEditor/Private/Views/List/ConsoleVariablesEditorList.cpp
--- a/Engine/Plugins/Editor/ConsoleVariablesEditor/Source/ConsoleVariablesEditor/Private/Views/List/ConsoleVariablesEditorList.cpp
+++ b/Engine/Plugins/Editor/ConsoleVariablesEditor/Source/ConsoleVariablesEditor/Private/Views/List/ConsoleVariablesEditorList.cpp
@@ -28,7 +28,7 @@ void FConsoleVariablesEditorList::RebuildList(const FString& InConsoleCommandToS
 {
 	if (ListWidget.IsValid())
 	{
-		ListWidget->RebuildList(InConsoleCommandToScrollTo, bShouldCacheValues);
+		ListWidget->RebuildListWithListMode(GetListMode(), InConsoleCommandToScrollTo, bShouldCacheValues);
 	}
 }
 
diff --git a/Engine/Plugins/Editor/ConsoleVariablesEditor/Source/ConsoleVariablesEditor/Private/Views/List/ConsoleVariablesEditorListFilters/ConsoleVariablesEditorListFilter_ModifiedVariables.h b/Engine/Plugins/Editor/ConsoleVariablesEditor/Source/ConsoleVariablesEditor/Private/Views/List/ConsoleVariablesEditorListFilters/ConsoleVariablesEditorListFilter_ModifiedVariables.h
deleted file mode 100644
diff --git a/Engine/Plugins/Editor/ConsoleVariablesEditor/Source/ConsoleVariablesEditor/Private/Views/List/ConsoleVariablesEditorListFilters/ConsoleVariablesEditorListFilter_SetByCurrentPreset.h b/Engine/Plugins/Editor/ConsoleVariablesEditor/Source/ConsoleVariablesEditor/Private/Views/List/ConsoleVariablesEditorListFilters/ConsoleVariablesEditorListFilter_SetByCurrentPreset.h
new file mode 100644
--- /dev/null
+++ b/Engine/Plugins/Editor/ConsoleVariablesEditor/Source/ConsoleVariablesEditor/Private/Views/List/ConsoleVariablesEditorListFilters/ConsoleVariablesEditorListFilter_SetByCurrentPreset.h
@@ -0,0 +1,46 @@
+// Copyright Epic Games, Inc. All Rights Reserved.
+
+#pragma once
+
+#include "IConsoleVariablesEditorListFilter.h"
+
+#define LOCTEXT_NAMESPACE "ConsoleVariablesEditor"
+
+class ConsoleVariablesEditorListFilter_SetByCurrentPreset : public IConsoleVariablesEditorListFilter
+{
+public:
+
+	ConsoleVariablesEditorListFilter_SetByCurrentPreset()
+	{
+		// This filter should be off by default
+		SetFilterActive(false);
+		SetFilterMatchType(EConsoleVariablesEditorListFilterMatchType::MatchAll);
+	}
+
+	virtual FString GetFilterName() override
+	{
+		return "Show Only Set By Current Preset";
+	}
+
+	virtual FText GetFilterButtonLabel() override
+	{
+		return LOCTEXT("ShowOnlySetByCurrentPresetFilter", "Show Only Set By Current Preset");
+	}
+
+	virtual FText GetFilterButtonToolTip() override
+	{
+		return LOCTEXT("ShowOnlySetByCurrentPresetFilterTooltip", "Show only rows that have been set by the current preset.");
+	}
+
+	virtual bool DoesItemPassFilter(const FConsoleVariablesEditorListRowPtr& InItem) override
+	{
+		if (InItem.IsValid(); const TSharedPtr<FConsoleVariablesEditorCommandInfo> PinnedCommand = InItem->GetCommandInfo().Pin())
+		{
+			return PinnedCommand->bChangedInCurrentPreset;
+		}
+
+		return false;
+	}
+};
+
+#undef LOCTEXT_NAMESPACE
diff --git a/Engine/Plugins/Editor/ConsoleVariablesEditor/Source/ConsoleVariablesEditor/Private/Views/List/ConsoleVariablesEditorListFilters/ConsoleVariablesEditorListFilter_SetInSession.h b/Engine/Plugins/Editor/ConsoleVariablesEditor/Source/ConsoleVariablesEditor/Private/Views/List/ConsoleVariablesEditorListFilters/ConsoleVariablesEditorListFilter_SetInSession.h
new file mode 100644
--- /dev/null
+++ b/Engine/Plugins/Editor/ConsoleVariablesEditor/Source/ConsoleVariablesEditor/Private/Views/List/ConsoleVariablesEditorListFilters/ConsoleVariablesEditorListFilter_SetInSession.h
@@ -0,0 +1,46 @@
+// Copyright Epic Games, Inc. All Rights Reserved.
+
+#pragma once
+
+#include "IConsoleVariablesEditorListFilter.h"
+
+#define LOCTEXT_NAMESPACE "ConsoleVariablesEditor"
+
+class ConsoleVariablesEditorListFilter_SetInSession : public IConsoleVariablesEditorListFilter
+{
+public:
+
+	ConsoleVariablesEditorListFilter_SetInSession()
+	{
+		// This filter should be off by default
+		SetFilterActive(false);
+		SetFilterMatchType(EConsoleVariablesEditorListFilterMatchType::MatchAll);
+	}
+
+	virtual FString GetFilterName() override
+	{
+		return "Show Only Set In Session";
+	}
+
+	virtual FText GetFilterButtonLabel() override
+	{
+		return LOCTEXT("ShowOnlySetInSessionFilter", "Show Only Set In Session");
+	}
+
+	virtual FText GetFilterButtonToolTip() override
+	{
+		return LOCTEXT("ShowOnlySetInSessionFilterTooltip", "Show only rows that have been set in the current session.");
+	}
+
+	virtual bool DoesItemPassFilter(const FConsoleVariablesEditorListRowPtr& InItem) override
+	{
+		if (InItem.IsValid(); const TSharedPtr<FConsoleVariablesEditorCommandInfo> PinnedCommand = InItem->GetCommandInfo().Pin())
+		{
+			return PinnedCommand->bChangedInCurrentPreset;
+		}
+
+		return false;
+	}
+};
+
+#undef LOCTEXT_NAMESPACE
diff --git a/Engine/Plugins/Editor/ConsoleVariablesEditor/Source/ConsoleVariablesEditor/Private/Views/List/ConsoleVariablesEditorListFilters/ConsoleVariablesEditorListFilter_ShowOnlyCommands.h b/Engine/Plugins/Editor/ConsoleVariablesEditor/Source/ConsoleVariablesEditor/Private/Views/List/ConsoleVariablesEditorListFilters/ConsoleVariablesEditorListFilter_ShowOnlyCommands.h
new file mode 100644
--- /dev/null
+++ b/Engine/Plugins/Editor/ConsoleVariablesEditor/Source/ConsoleVariablesEditor/Private/Views/List/ConsoleVariablesEditorListFilters/ConsoleVariablesEditorListFilter_ShowOnlyCommands.h
@@ -0,0 +1,46 @@
+// Copyright Epic Games, Inc. All Rights Reserved.
+
+#pragma once
+
+#include "IConsoleVariablesEditorListFilter.h"
+
+#define LOCTEXT_NAMESPACE "ConsoleVariablesEditor"
+
+class ConsoleVariablesEditorListFilter_ShowOnlyCommands : public IConsoleVariablesEditorListFilter
+{
+public:
+
+	ConsoleVariablesEditorListFilter_ShowOnlyCommands()
+	{
+		// This filter should be off by default
+		SetFilterActive(false);
+		SetFilterMatchType(EConsoleVariablesEditorListFilterMatchType::MatchAll);
+	}
+
+	virtual FString GetFilterName() override
+	{
+		return "Show Only Commands";
+	}
+
+	virtual FText GetFilterButtonLabel() override
+	{
+		return LOCTEXT("ShowOnlyCommandsFilter", "Show Only Commands");
+	}
+
+	virtual FText GetFilterButtonToolTip() override
+	{
+		return LOCTEXT("ShowOnlyCommandsFilterTooltip", "Show only rows that represent console commands, like 'stat unit'.");
+	}
+
+	virtual bool DoesItemPassFilter(const FConsoleVariablesEditorListRowPtr& InItem) override
+	{
+		if (InItem.IsValid(); const TSharedPtr<FConsoleVariablesEditorCommandInfo> PinnedCommand = InItem->GetCommandInfo().Pin())
+		{
+			return PinnedCommand->ObjectType != FConsoleVariablesEditorCommandInfo::EConsoleObjectType::Variable;
+		}
+
+		return false;
+	}
+};
+
+#undef LOCTEXT_NAMESPACE
diff --git a/Engine/Plugins/Editor/ConsoleVariablesEditor/Source/ConsoleVariablesEditor/Private/Views/List/ConsoleVariablesEditorListFilters/ConsoleVariablesEditorListFilter_ShowOnlyModifiedVariables.h b/Engine/Plugins/Editor/ConsoleVariablesEditor/Source/ConsoleVariablesEditor/Private/Views/List/ConsoleVariablesEditorListFilters/ConsoleVariablesEditorListFilter_ShowOnlyModifiedVariables.h
new file mode 100644
--- /dev/null
+++ b/Engine/Plugins/Editor/ConsoleVariablesEditor/Source/ConsoleVariablesEditor/Private/Views/List/ConsoleVariablesEditorListFilters/ConsoleVariablesEditorListFilter_ShowOnlyModifiedVariables.h
@@ -0,0 +1,46 @@
+// Copyright Epic Games, Inc. All Rights Reserved.
+
+#pragma once
+
+#include "IConsoleVariablesEditorListFilter.h"
+
+#define LOCTEXT_NAMESPACE "ConsoleVariablesEditor"
+
+class ConsoleVariablesEditorListFilter_ShowOnlyModifiedVariables : public IConsoleVariablesEditorListFilter
+{
+public:
+
+	ConsoleVariablesEditorListFilter_ShowOnlyModifiedVariables()
+	{
+		// This filter should be off by default
+		SetFilterActive(false);
+		SetFilterMatchType(EConsoleVariablesEditorListFilterMatchType::MatchAll);
+	}
+
+	virtual FString GetFilterName() override
+	{
+		return "Show Only Modified Variables";
+	}
+
+	virtual FText GetFilterButtonLabel() override
+	{
+		return LOCTEXT("ShowOnlyModifiedVariablesFilter", "Show Only Modified");
+	}
+
+	virtual FText GetFilterButtonToolTip() override
+	{
+		return LOCTEXT("ShowOnlyModifiedVariablesFilterTooltip", "Show only rows that have a current value that differs from the preset value.");
+	}
+
+	virtual bool DoesItemPassFilter(const FConsoleVariablesEditorListRowPtr& InItem) override
+	{
+		if (InItem.IsValid())
+		{
+			return InItem->DoesCurrentValueDifferFromPresetValue();
+		}
+
+		return false;
+	}
+};
+
+#undef LOCTEXT_NAMESPACE
diff --git a/Engine/Plugins/Editor/ConsoleVariablesEditor/Source/ConsoleVariablesEditor/Private/Views/List/ConsoleVariablesEditorListFilters/ConsoleVariablesEditorListFilter_ShowOnlyVariables.h b/Engine/Plugins/Editor/ConsoleVariablesEditor/Source/ConsoleVariablesEditor/Private/Views/List/ConsoleVariablesEditorListFilters/ConsoleVariablesEditorListFilter_ShowOnlyVariables.h
new file mode 100644
--- /dev/null
+++ b/Engine/Plugins/Editor/ConsoleVariablesEditor/Source/ConsoleVariablesEditor/Private/Views/List/ConsoleVariablesEditorListFilters/ConsoleVariablesEditorListFilter_ShowOnlyVariables.h
@@ -0,0 +1,46 @@
+// Copyright Epic Games, Inc. All Rights Reserved.
+
+#pragma once
+
+#include "IConsoleVariablesEditorListFilter.h"
+
+#define LOCTEXT_NAMESPACE "ConsoleVariablesEditor"
+
+class ConsoleVariablesEditorListFilter_ShowOnlyVariables : public IConsoleVariablesEditorListFilter
+{
+public:
+
+	ConsoleVariablesEditorListFilter_ShowOnlyVariables()
+	{
+		// This filter should be off by default
+		SetFilterActive(false);
+		SetFilterMatchType(EConsoleVariablesEditorListFilterMatchType::MatchAll);
+	}
+
+	virtual FString GetFilterName() override
+	{
+		return "Show Only Variables";
+	}
+
+	virtual FText GetFilterButtonLabel() override
+	{
+		return LOCTEXT("ShowOnlyVariablesFilter", "Show Only Variables");
+	}
+
+	virtual FText GetFilterButtonToolTip() override
+	{
+		return LOCTEXT("ShowOnlyVariablesFilterTooltip", "Show only rows that represent console variables.");
+	}
+
+	virtual bool DoesItemPassFilter(const FConsoleVariablesEditorListRowPtr& InItem) override
+	{
+		if (InItem.IsValid(); const TSharedPtr<FConsoleVariablesEditorCommandInfo> PinnedCommand = InItem->GetCommandInfo().Pin())
+		{
+			return PinnedCommand->ObjectType == FConsoleVariablesEditorCommandInfo::EConsoleObjectType::Variable;
+		}
+
+		return false;
+	}
+};
+
+#undef LOCTEXT_NAMESPACE
diff --git a/Engine/Plugins/Editor/ConsoleVariablesEditor/Source/ConsoleVariablesEditor/Private/Views/List/ConsoleVariablesEditorListFilters/ConsoleVariablesEditorListFilter_Source.h b/Engine/Plugins/Editor/ConsoleVariablesEditor/Source/ConsoleVariablesEditor/Private/Views/List/ConsoleVariablesEditorListFilters/ConsoleVariablesEditorListFilter_Source.h
new file mode 100644
--- /dev/null
+++ b/Engine/Plugins/Editor/ConsoleVariablesEditor/Source/ConsoleVariablesEditor/Private/Views/List/ConsoleVariablesEditorListFilters/ConsoleVariablesEditorListFilter_Source.h
@@ -0,0 +1,49 @@
+// Copyright Epic Games, Inc. All Rights Reserved.
+
+#pragma once
+
+#include "IConsoleVariablesEditorListFilter.h"
+
+#define LOCTEXT_NAMESPACE "ConsoleVariablesEditor"
+
+class FConsoleVariablesEditorListFilter_Source : public IConsoleVariablesEditorListFilter
+{
+public:
+
+	FConsoleVariablesEditorListFilter_Source(const EConsoleVariableFlags InSourceFlag)
+	: SourceFlag(InSourceFlag){}
+
+	virtual FString GetFilterName() override
+	{
+		return FConsoleVariablesEditorCommandInfo::ConvertConsoleVariableSetByFlagToText(SourceFlag).ToString();
+	}
+
+	virtual FText GetFilterButtonLabel() override
+	{
+		return FText::Format(LOCTEXT("ShowSourceTextFilterFormat", "Show {0}"), FText::FromString(GetFilterName()));
+	}
+
+	virtual FText GetFilterButtonToolTip() override
+	{
+		return FText::Format(
+			LOCTEXT("ShowSourceTextFilterTooltipFormat", "Show rows that have a Source field matching '{0}'"),
+			FText::FromString(GetFilterName()));
+	}
+
+	virtual bool DoesItemPassFilter(const FConsoleVariablesEditorListRowPtr& InItem) override
+	{
+		if (InItem.IsValid(); const TSharedPtr<FConsoleVariablesEditorCommandInfo> PinnedCommand = InItem->GetCommandInfo().Pin())
+		{
+			const EConsoleVariableFlags Source = PinnedCommand->GetSource();
+			return Source == SourceFlag;
+		}
+
+		return false;
+	}
+	
+private:
+
+	EConsoleVariableFlags SourceFlag;
+};
+
+#undef LOCTEXT_NAMESPACE
diff --git a/Engine/Plugins/Editor/ConsoleVariablesEditor/Source/ConsoleVariablesEditor/Private/Views/List/ConsoleVariablesEditorListFilters/ConsoleVariablesEditorListFilter_SourceText.h b/Engine/Plugins/Editor/ConsoleVariablesEditor/Source/ConsoleVariablesEditor/Private/Views/List/ConsoleVariablesEditorListFilters/ConsoleVariablesEditorListFilter_SourceText.h
deleted file mode 100644
diff --git a/Engine/Plugins/Editor/ConsoleVariablesEditor/Source/ConsoleVariablesEditor/Private/Views/List/ConsoleVariablesEditorListRow.cpp b/Engine/Plugins/Editor/ConsoleVariablesEditor/Source/ConsoleVariablesEditor/Private/Views/List/ConsoleVariablesEditorListRow.cpp
--- a/Engine/Plugins/Editor/ConsoleVariablesEditor/Source/ConsoleVariablesEditor/Private/Views/List/ConsoleVariablesEditorListRow.cpp
+++ b/Engine/Plugins/Editor/ConsoleVariablesEditor/Source/ConsoleVariablesEditor/Private/Views/List/ConsoleVariablesEditorListRow.cpp
@@ -122,8 +122,12 @@ void FConsoleVariablesEditorListRow::SetShouldExpandAllChildren(const bool bNewS
 
 void FConsoleVariablesEditorListRow::ResetToStartupValueAndSource() const
 {
-	GetCommandInfo().Pin()->ExecuteCommand(GetCommandInfo().Pin()->StartupValueAsString);
-	GetCommandInfo().Pin()->SetSourceFlag(GetCommandInfo().Pin()->StartupSource);
+	if (const TSharedPtr<FConsoleVariablesEditorCommandInfo> PinnedCommand = GetCommandInfo().Pin())
+	{
+		PinnedCommand->SetSourceFlag(PinnedCommand->StartupSource);
+		PinnedCommand->ExecuteCommand(PinnedCommand->StartupValueAsString);
+		PinnedCommand->SetIfChangedInCurrentPreset(false);
+	}
 }
 
 const FString& FConsoleVariablesEditorListRow::GetPresetValue() const
@@ -343,7 +347,7 @@ FReply FConsoleVariablesEditorListRow::OnActionButtonClicked()
 
 		EditableAsset->RemoveConsoleVariable(CommandName);
 
-		ListViewPtr.Pin()->RebuildList();
+		ListViewPtr.Pin()->RebuildListWithListMode(ListViewPtr.Pin()->GetListModelPtr().Pin()->GetListMode());
 	}
 
 	return FReply::Handled();
diff --git a/Engine/Plugins/Editor/ConsoleVariablesEditor/Source/ConsoleVariablesEditor/Private/Views/List/SConsoleVariablesEditorList.cpp b/Engine/Plugins/Editor/ConsoleVariablesEditor/Source/ConsoleVariablesEditor/Private/Views/List/SConsoleVariablesEditorList.cpp
--- a/Engine/Plugins/Editor/ConsoleVariablesEditor/Source/ConsoleVariablesEditor/Private/Views/List/SConsoleVariablesEditorList.cpp
+++ b/Engine/Plugins/Editor/ConsoleVariablesEditor/Source/ConsoleVariablesEditor/Private/Views/List/SConsoleVariablesEditorList.cpp
@@ -2,9 +2,13 @@
 
 #include "SConsoleVariablesEditorList.h"
 
-#include "ConsoleVariablesEditorList.h"
-#include "ConsoleVariablesEditorListFilters/ConsoleVariablesEditorListFilter_ModifiedVariables.h"
-#include "ConsoleVariablesEditorListFilters/ConsoleVariablesEditorListFilter_SourceText.h"
+// Filters
+#include "ConsoleVariablesEditorListFilters/ConsoleVariablesEditorListFilter_ShowOnlyModifiedVariables.h"
+#include "ConsoleVariablesEditorListFilters/ConsoleVariablesEditorListFilter_SetInSession.h"
+#include "ConsoleVariablesEditorListFilters/ConsoleVariablesEditorListFilter_ShowOnlyCommands.h"
+#include "ConsoleVariablesEditorListFilters/ConsoleVariablesEditorListFilter_ShowOnlyVariables.h"
+#include "ConsoleVariablesEditorListFilters/ConsoleVariablesEditorListFilter_Source.h"
+
 #include "ConsoleVariablesEditorModule.h"
 #include "ConsoleVariablesEditorProjectSettings.h"
 #include "ConsoleVariablesEditorStyle.h"
@@ -167,12 +171,10 @@ void SConsoleVariablesEditorList::Construct(const FArguments& InArgs, TSharedRef
 				.IsChecked(false)
 				.OnCheckStateChanged_Lambda([this] (ECheckBoxState NewCheckState)
 				{
-					ListModelPtr.Pin()->SetListMode(FConsoleVariablesEditorList::EConsoleVariablesEditorListMode::Preset);
-
 					GlobalSearchesContainer->ClearChildren();
 					CurrentGlobalSearches.Empty();
 
-					RebuildList();
+					RebuildListWithListMode(FConsoleVariablesEditorList::EConsoleVariablesEditorListMode::Preset);
 
 					RemoveGlobalSearchesButtonPtr->SetIsChecked(false);
 				})
@@ -231,15 +233,23 @@ void SConsoleVariablesEditorList::Construct(const FArguments& InArgs, TSharedRef
 				.Justification(ETextJustify::Center)
 				.Text_Lambda([this]()
 				{
-					if (VisibleTreeViewObjects.Num() == TreeViewRootObjects.Num())
+					if (VisibleTreeViewObjects.Num() == 0)
 					{
-						// Global Search Empty List (without filter)
-                    	return LOCTEXT("ConsoleVariablesEditorList_EmptyListGlobalSearch", "No matching console variables found in Unreal Engine.\n\nCheck your search criteria.");
-						
+						if (TreeViewRootObjects.Num() == 0)
+						{
+							if (GetListModelPtr().Pin()->GetListMode() == FConsoleVariablesEditorList::EConsoleVariablesEditorListMode::GlobalSearch)
+							{
+								// Global Search Empty List (without filter)
+								return LOCTEXT("ConsoleVariablesEditorList_EmptyListGlobalSearch", "No matching console variables found in Unreal Engine.\n\nCheck your search criteria.");
+							}
+
+							// Preset Empty List (without Filter)
+							return LOCTEXT("ConsoleVariablesEditorList_EmptyListPresetWithoutFilter", "Type in a variable to the console or <RichTextBlock.Bold>Search All</> to add to your list.");
+						}
 					}
 
-					// Empty List (with filter)
-                    return LOCTEXT("ConsoleVariablesEditorList_EmptyListWithFilter", "No matching console variables in your list.\n\nCheck your filter or <RichTextBlock.Bold>Search All</> console variables instead.");
+					// Preset Empty List (with filter)
+                    return LOCTEXT("ConsoleVariablesEditorList_EmptyListPresetWithFilter", "No matching console variables in your list.\n\nCheck your filter or <RichTextBlock.Bold>Search All</> console variables instead.");
 				})
 			]
 		]
@@ -266,6 +276,8 @@ SConsoleVariablesEditorList::~SConsoleVariablesEditorList()
 	TreeViewPtr.Reset();
 	VisibleTreeViewObjects.Reset();
 	LastPresetObjects.Reset();
+	
+	CachedCommandStates.Empty();
 }
 
 FReply SConsoleVariablesEditorList::TryEnterGlobalSearch(const FString& SearchString)
@@ -279,11 +291,10 @@ FReply SConsoleVariablesEditorList::TryEnterGlobalSearch(const FString& SearchSt
 			TEXT("%hs: Global search request is empty. Exiting Global Search."),
 			__FUNCTION__, *SearchString);
 
-		// Return to preset mode if in global search mode then rebuild list
+		// Return to preset mode if in global search mode and rebuild list
 		if (ListModelPtr.Pin()->GetListMode() == FConsoleVariablesEditorList::EConsoleVariablesEditorListMode::GlobalSearch)
 		{
-			ListModelPtr.Pin()->SetListMode(FConsoleVariablesEditorList::EConsoleVariablesEditorListMode::Preset);
-			RebuildList();
+			RebuildListWithListMode(FConsoleVariablesEditorList::EConsoleVariablesEditorListMode::Preset);
 		}
 
 		return ReturnValue;
@@ -322,13 +333,6 @@ FReply SConsoleVariablesEditorList::TryEnterGlobalSearch(const FString& SearchSt
 	
 	const bool bFoundMatches = ConsoleVariablesEditorModule.PopulateGlobalSearchAssetWithVariablesMatchingTokens(OutTokens);
 	{
-		// If we were in Preset mode before entering global search, cache the existing tree objects to maintain state
-		if (ListModelPtr.Pin()->GetListMode() == FConsoleVariablesEditorList::EConsoleVariablesEditorListMode::Preset)
-		{
-			LastPresetObjects = TreeViewRootObjects;
-			ListModelPtr.Pin()->SetListMode(FConsoleVariablesEditorList::EConsoleVariablesEditorListMode::GlobalSearch);
-		}
-
 		// Convert tokens to global search toggle buttons
 		for (const FString& TokenString : OutTokens)
 		{
@@ -361,7 +365,7 @@ FReply SConsoleVariablesEditorList::TryEnterGlobalSearch(const FString& SearchSt
 			__FUNCTION__, *SearchString);
 	}
 
-	RebuildList();
+	RebuildListWithListMode(FConsoleVariablesEditorList::EConsoleVariablesEditorListMode::GlobalSearch);
 
 	return ReturnValue;
 }
@@ -397,20 +401,35 @@ void SConsoleVariablesEditorList::RefreshGlobalSearchWidgets()
 	}
 }
 
-void SConsoleVariablesEditorList::RebuildList(const FString& InConsoleCommandToScrollTo, bool bShouldCacheValues)
+void SConsoleVariablesEditorList::RebuildListWithListMode(
+	FConsoleVariablesEditorList::EConsoleVariablesEditorListMode NewListMode, const FString& InConsoleCommandToScrollTo, bool bShouldCacheValues)
 {
 	if (bShouldCacheValues)
 	{
 		CacheCurrentListItemData();
 	}
 
+	// If we're in Preset mode and changing to something else, cache the existing tree objects to maintain state
+	if (GetListModelPtr().Pin()->GetListMode() == FConsoleVariablesEditorList::EConsoleVariablesEditorListMode::Preset &&
+		NewListMode != FConsoleVariablesEditorList::EConsoleVariablesEditorListMode::Preset)
+	{
+		LastPresetObjects = TreeViewRootObjects;
+	}
+	
+	GetListModelPtr().Pin()->SetListMode(NewListMode);
+
 	// Skip execution on load if we've cached the previous values
-	GenerateTreeView(!bShouldCacheValues);
+	GenerateTreeView(bShouldCacheValues);
 
 	if (bShouldCacheValues)
 	{
 		RestorePreviousListItemData();
 	}
+	else
+	{
+		// If we didn't want to cache, clear the previous cache
+		CachedCommandStates.Empty();
+	}
 
 	RefreshList();
 
@@ -451,9 +470,9 @@ void SConsoleVariablesEditorList::RefreshList()
 
 		// Refresh the header's check state
 		OnListItemCheckBoxStateChange(ECheckBoxState::Undetermined);
-
-		FindVisibleObjectsAndRequestTreeRefresh();
 	}
+
+	FindVisibleObjectsAndRequestTreeRefresh();
 }
 
 TArray<FConsoleVariablesEditorListRowPtr> SConsoleVariablesEditorList::GetSelectedTreeViewItems() const
@@ -839,28 +858,37 @@ TSharedPtr<SHeaderRow> SConsoleVariablesEditorList::GenerateHeaderRow()
 
 void SConsoleVariablesEditorList::SetupFilters()
 {
-	TArray<FString> SourceFilterTypes =
+	TArray<EConsoleVariableFlags> SourceFilterTypes =
 	{
-		"Constructor",
-		"Scalability",
-		"Game Setting",
-		"Project Setting",
-		"System Settings ini",
-		"Device Profile",
-		"Game Override",
-		"Console Variables ini",
-		"Command line",
-		"Code",
-		"Console"
+		ECVF_SetByConstructor,
+		ECVF_SetByScalability,
+		ECVF_SetByGameSetting,
+		ECVF_SetByProjectSetting,
+		ECVF_SetBySystemSettingsIni,
+		ECVF_SetByDeviceProfile,
+		ECVF_SetByGameOverride,
+		ECVF_SetByConsoleVariablesIni,
+		ECVF_SetByCommandline,
+		ECVF_SetByCode,
+		ECVF_SetByConsole
 	};
 
-	for (const FString& Type : SourceFilterTypes)
+	for (const EConsoleVariableFlags Type : SourceFilterTypes)
 	{
-		ShowFilters.Add(MakeShared<FConsoleVariablesEditorListFilter_SourceText>(Type));
+		ShowFilters.Add(MakeShared<FConsoleVariablesEditorListFilter_Source>(Type));
 	}
 
 	// Add Show Only Modified filter
-	ShowFilters.Add(MakeShared<ConsoleVariablesEditorListFilter_ModifiedVariables>());
+	ShowFilters.Add(MakeShared<ConsoleVariablesEditorListFilter_ShowOnlyModifiedVariables>());
+
+	// Add Show Only Set By Current Preset filter
+	ShowFilters.Add(MakeShared<ConsoleVariablesEditorListFilter_SetInSession>());
+
+	// Add Show Only Variables filter
+	ShowFilters.Add(MakeShared<ConsoleVariablesEditorListFilter_ShowOnlyVariables>());
+
+	// Add Show Only Commands filter
+	ShowFilters.Add(MakeShared<ConsoleVariablesEditorListFilter_ShowOnlyCommands>());
 }
 
 TSharedRef<SWidget> SConsoleVariablesEditorList::BuildShowOptionsMenu()
@@ -902,18 +930,6 @@ TSharedRef<SWidget> SConsoleVariablesEditorList::BuildShowOptionsMenu()
 	
 	ShowOptionsMenuBuilder.BeginSection("", LOCTEXT("ShowOptions_SortSectionHeading", "Sort"));
 	{
-		// Add commands
-
-		// Save this for later when folders are added
-		/*ShowOptionsMenuBuilder.AddMenuEntry(
-			LOCTEXT("CollapseAll", "Collapse All"),
-			LOCTEXT("ConsoleVariablesEditorList_CollapseAll_Tooltip", "Collapse all expanded actor groups in the Modified Actors list."),
-			FSlateIcon(),
-			FUIAction(FExecuteAction::CreateRaw(this, &SConsoleVariablesEditorList::SetAllGroupsCollapsed)),
-			NAME_None,
-			EUserInterfaceActionType::Button
-		);*/
-		
 		ShowOptionsMenuBuilder.AddMenuEntry(
 			LOCTEXT("SetSortOrder", "Set Sort Order"),
 			LOCTEXT("ConsoleVariablesEditorList_SetSortOrder_Tooltip", "Makes the current order of the variables list the saved order."),
@@ -998,27 +1014,32 @@ void SConsoleVariablesEditorList::OnListViewSearchTextChanged(const FText& Text)
 
 void SConsoleVariablesEditorList::CacheCurrentListItemData()
 {
-	CachedCommandStates.Empty();
-
-	// We only want preset items, not global search
 	const TArray<FConsoleVariablesEditorListRowPtr>& ListItems =
-		ListModelPtr.Pin()->GetListMode() == FConsoleVariablesEditorList::EConsoleVariablesEditorListMode::Preset ?
-			TreeViewRootObjects : LastPresetObjects;
+		GetListModelPtr().Pin()->GetListMode() ==
+			FConsoleVariablesEditorList::EConsoleVariablesEditorListMode::Preset ? TreeViewRootObjects : LastPresetObjects;
+
+	CachedCommandStates.Empty(CachedCommandStates.Num());
 
 	for (const TSharedPtr<FConsoleVariablesEditorListRow>& Item : ListItems)
 	{
 		if (const TSharedPtr<FConsoleVariablesEditorCommandInfo>& CommandInfo = Item->GetCommandInfo().Pin())
 		{
-			if (const IConsoleVariable* AsVariable = CommandInfo->GetConsoleVariablePtr())
-			{
-				CachedCommandStates.Add(
-					{
-						CommandInfo->Command,
-						AsVariable->GetString(),
-						Item->GetWidgetCheckedState()
-					}
-				);
-			}
+			const FString& CommandName = CommandInfo->Command;
+				
+			/*CachedCommandStates.SetNum(Algo::RemoveIf(
+				CachedCommandStates,
+				[CommandName](const FConsoleVariablesEditorAssetSaveData& CachedData)
+				{
+					return CachedData.CommandName.Equals(CommandName);
+				}));*/
+			
+			CachedCommandStates.Add(
+				{
+					CommandInfo->Command,
+					Item->GetCachedValue(),
+					Item->GetWidgetCheckedState()
+				}
+			);
 		}
 	}
 }
@@ -1029,7 +1050,9 @@ void SConsoleVariablesEditorList::RestorePreviousListItemData()
 	{
 		if (const TSharedPtr<FConsoleVariablesEditorCommandInfo>& CommandInfo = Item->GetCommandInfo().Pin())
 		{
-			if (CommandInfo->GetConsoleVariablePtr())
+			FConsoleVariablesEditorCommandInfo::EConsoleObjectType CommandType = CommandInfo->ObjectType;
+			
+			if (CommandInfo->GetConsoleObjectPtr())
 			{
 				const FString& CommandName = CommandInfo->Command;
 				
@@ -1040,17 +1063,15 @@ void SConsoleVariablesEditorList::RestorePreviousListItemData()
 						return CachedData.CommandName.Equals(CommandName);
 					}))
 				{
-					CommandInfo->ExecuteCommand((*Match).CommandValueAsString);
+					Item->SetCachedValue((*Match).CommandValueAsString);
 					Item->SetWidgetCheckedState((*Match).CheckedState);
 				}
 			}
 		}
 	}
-
-	CachedCommandStates.Empty();
 }
 
-void SConsoleVariablesEditorList::GenerateTreeView(const bool bExecuteCommandsAsTheyAreLoaded)
+void SConsoleVariablesEditorList::GenerateTreeView(const bool bSkipExecutionOfCachedCommands)
 {	
 	if (!ensure(TreeViewPtr.IsValid()))
 	{
@@ -1059,6 +1080,9 @@ void SConsoleVariablesEditorList::GenerateTreeView(const bool bExecuteCommandsAs
 	
 	FlushMemory(true);
 
+	const bool bIsPresetMode =
+		ListModelPtr.Pin()->GetListMode() == FConsoleVariablesEditorList::EConsoleVariablesEditorListMode::Preset;
+
 	FConsoleVariablesEditorModule& ConsoleVariablesEditorModule = FConsoleVariablesEditorModule::Get();
 
 	TObjectPtr<UConsoleVariablesAsset> EditableAsset = nullptr;
@@ -1076,6 +1100,7 @@ void SConsoleVariablesEditorList::GenerateTreeView(const bool bExecuteCommandsAs
 		default:
 			checkNoEntry();
 	}
+	
 	check(EditableAsset);
 
 	for (const FConsoleVariablesEditorAssetSaveData& SavedCommand : EditableAsset->GetSavedCommands())
@@ -1084,6 +1109,7 @@ void SConsoleVariablesEditorList::GenerateTreeView(const bool bExecuteCommandsAs
 		TSharedPtr<FConsoleVariablesEditorCommandInfo> CommandInfo =
 			ConsoleVariablesEditorModule.FindCommandInfoByName(SavedCommand.CommandName).Pin();
 
+		// Make new CommandInfo if one doesn't exist
 		if (!CommandInfo.IsValid())
 		{
 			CommandInfo = MakeShared<FConsoleVariablesEditorCommandInfo>(SavedCommand.CommandName);
@@ -1111,22 +1137,6 @@ void SConsoleVariablesEditorList::GenerateTreeView(const bool bExecuteCommandsAs
 				const ECheckBoxState NewCheckedState =
 					SavedCommand.CheckedState == ECheckBoxState::Unchecked ?
 							ECheckBoxState::Unchecked : ECheckBoxState::Checked;
-
-				if (bExecuteCommandsAsTheyAreLoaded &&
-					ListModelPtr.Pin()->GetListMode() == FConsoleVariablesEditorList::EConsoleVariablesEditorListMode::Preset)
-				{
-					// If the row is checked and the saved value differs from the current value,
-					// execute the command with the saved value 
-					if (CommandInfo->ObjectType == FConsoleVariablesEditorCommandInfo::EConsoleObjectType::Command ||
-						(NewCheckedState == ECheckBoxState::Checked &&
-						CommandInfo->IsCurrentValueDifferentFromInputValue(SavedCommand.CommandValueAsString)))
-					{
-						if (!SavedCommand.CommandValueAsString.IsEmpty())
-						{
-							CommandInfo->ExecuteCommand(SavedCommand.CommandValueAsString);
-						}
-					}
-				}
 		
 				RowToAdd = 
 				   MakeShared<FConsoleVariablesEditorListRow>(
@@ -1134,15 +1144,42 @@ void SConsoleVariablesEditorList::GenerateTreeView(const bool bExecuteCommandsAs
 						   FConsoleVariablesEditorListRow::SingleCommand, 
 						   NewCheckedState, SharedThis(this), TreeViewRootObjects.Num(), nullptr);
 			}
-			
-			TreeViewRootObjects.Add(RowToAdd);
+
+			if (RowToAdd)
+			{
+				bool bIsEligibleForExecution =
+					bIsPresetMode &&
+					CommandInfo->ObjectType != FConsoleVariablesEditorCommandInfo::EConsoleObjectType::NullObject &&
+					RowToAdd->IsRowChecked() &&
+					!SavedCommand.CommandValueAsString.IsEmpty() &&
+					CommandInfo->IsCurrentValueDifferentFromInputValue(SavedCommand.CommandValueAsString);
+				
+				if (bIsEligibleForExecution && bSkipExecutionOfCachedCommands)
+				{
+					// If we want to skip execution of cached commands, we need to check the cached commands for a match
+					// first. If one isn't found, then we'll execute.
+					bIsEligibleForExecution =
+						!CachedCommandStates.ContainsByPredicate(
+							[this, &CommandInfo](const FConsoleVariablesEditorAssetSaveData& Comparator)
+							{
+								return Comparator.CommandName.Equals(CommandInfo->Command);
+							});
+				}
+
+				if (bIsEligibleForExecution)
+				{
+					CommandInfo->ExecuteCommand(SavedCommand.CommandValueAsString);
+				}
+				
+				TreeViewRootObjects.Add(RowToAdd);
+			}
 		}
 	}
 
 	// Now clear out the last preset cache if the list is in preset mode
-	if (ListModelPtr.Pin()->GetListMode() == FConsoleVariablesEditorList::EConsoleVariablesEditorListMode::Preset)
+	if (bIsPresetMode)
 	{
-		LastPresetObjects.Empty();
+		LastPresetObjects.Empty(LastPresetObjects.Num());
 	}
 
 	TreeViewPtr->RequestTreeRefresh();
diff --git a/Engine/Plugins/Editor/ConsoleVariablesEditor/Source/ConsoleVariablesEditor/Private/Views/List/SConsoleVariablesEditorList.h b/Engine/Plugins/Editor/ConsoleVariablesEditor/Source/ConsoleVariablesEditor/Private/Views/List/SConsoleVariablesEditorList.h
--- a/Engine/Plugins/Editor/ConsoleVariablesEditor/Source/ConsoleVariablesEditor/Private/Views/List/SConsoleVariablesEditorList.h
+++ b/Engine/Plugins/Editor/ConsoleVariablesEditor/Source/ConsoleVariablesEditor/Private/Views/List/SConsoleVariablesEditorList.h
@@ -3,9 +3,11 @@
 #pragma once
 
 #include "ConsoleVariablesAsset.h"
+#include "ConsoleVariablesEditorList.h"
 #include "ConsoleVariablesEditorListRow.h"
 #include "ConsoleVariablesEditorListFilters/IConsoleVariablesEditorListFilter.h"
 
+#include "Misc/EnumRange.h"
 #include "Widgets/SCompoundWidget.h"
 #include "Widgets/Layout/SGridPanel.h"
 #include "Widgets/Views/SHeaderRow.h"
@@ -49,7 +51,7 @@ public:
 	 * Regenerate the list items and refresh the list. Call when adding or removing variables.
 	 * @param bShouldCacheValues If true, the current list's current values will be cached and then restored when the list is rebuilt. Otherwise preset values will be used.
 	 */
-	void RebuildList(const FString& InConsoleCommandToScrollTo = "", bool bShouldCacheValues = true);
+	void RebuildListWithListMode(FConsoleVariablesEditorList::EConsoleVariablesEditorListMode NewListMode, const FString& InConsoleCommandToScrollTo = "", bool bShouldCacheValues = true);
 
 	/**
 	 * Refresh filters and sorting.
@@ -152,7 +154,7 @@ private:
 
 	void CacheCurrentListItemData();
 	void RestorePreviousListItemData();
-	void GenerateTreeView(const bool bExecuteCommandsAsTheyAreLoaded = true);
+	void GenerateTreeView(const bool bSkipExecutionOfCachedCommands = true);
 	void FindVisibleTreeViewObjects();
 	void FindVisibleObjectsAndRequestTreeRefresh();
 	
diff --git a/Engine/Plugins/Editor/ConsoleVariablesEditor/Source/ConsoleVariablesEditor/Private/Views/List/SConsoleVariablesEditorListValueInput.cpp b/Engine/Plugins/Editor/ConsoleVariablesEditor/Source/ConsoleVariablesEditor/Private/Views/List/SConsoleVariablesEditorListValueInput.cpp
--- a/Engine/Plugins/Editor/ConsoleVariablesEditor/Source/ConsoleVariablesEditor/Private/Views/List/SConsoleVariablesEditorListValueInput.cpp
+++ b/Engine/Plugins/Editor/ConsoleVariablesEditor/Source/ConsoleVariablesEditor/Private/Views/List/SConsoleVariablesEditorListValueInput.cpp
@@ -56,7 +56,9 @@ TSharedRef<SConsoleVariablesEditorListValueInput> SConsoleVariablesEditorListVal
 		}
 
 		// For Commands
-		return SNew(SConsoleVariablesEditorListValueInput_Command, InRow, PinnedItem->GetPresetValue());
+		const FString CachedValue = PinnedItem->GetCachedValue();
+		return SNew(SConsoleVariablesEditorListValueInput_Command, InRow,
+			CachedValue.IsEmpty() ? PinnedItem->GetPresetValue() : CachedValue);
 	}
 
 	// fallback
diff --git a/Engine/Plugins/Editor/ConsoleVariablesEditor/Source/ConsoleVariablesEditor/Private/Views/MainPanel/ConsoleVariablesEditorMainPanel.cpp b/Engine/Plugins/Editor/ConsoleVariablesEditor/Source/ConsoleVariablesEditor/Private/Views/MainPanel/ConsoleVariablesEditorMainPanel.cpp
--- a/Engine/Plugins/Editor/ConsoleVariablesEditor/Source/ConsoleVariablesEditor/Private/Views/MainPanel/ConsoleVariablesEditorMainPanel.cpp
+++ b/Engine/Plugins/Editor/ConsoleVariablesEditor/Source/ConsoleVariablesEditor/Private/Views/MainPanel/ConsoleVariablesEditorMainPanel.cpp
@@ -60,7 +60,7 @@ void FConsoleVariablesEditorMainPanel::AddConsoleObjectToCurrentPreset(
 	{
 		Asset->AddOrSetConsoleObjectSavedData(
 			{
-				InConsoleCommand,
+				InConsoleCommand.TrimStartAndEnd(),
 				InValue,
 				ECheckBoxState::Checked
 			}
@@ -82,12 +82,33 @@ FReply FConsoleVariablesEditorMainPanel::ValidateConsoleInputAndAddToCurrentPres
 	}
 	
 	const FString CommandString = CommittedText.ToString().TrimStartAndEnd();
-	FString CommandKey;
-	FString ValueString;
+
+	if (CommandString.IsEmpty())
+	{
+		UE_LOG(LogConsoleVariablesEditor, Warning, TEXT("%hs: Input is blank."), __FUNCTION__);
+
+		return FReply::Unhandled();
+	}
 	
+	FString CommandKey; // The actual variable name, like r.ScreenPercentage
+	FString ValueString; // The variable value, usually after the CommandKey and a single space
+	FString AdditionalParams; // Some variables have multiple parameters between the CommandKey and ValueString
+
 	if (CommandString.Contains(" "))
 	{
-		CommandString.Split(TEXT(" "), &CommandKey, &ValueString);
+		// We want to treat the last string after the last space as the value
+		// This allows commands with multiple parameters to be treated as separate list items
+		const int32 IndexOfLastSpace = CommandString.Find(" ", ESearchCase::IgnoreCase, ESearchDir::FromEnd);
+		CommandKey = CommandString.Left(IndexOfLastSpace).TrimStartAndEnd();
+		ValueString = CommandString.RightChop(IndexOfLastSpace).TrimStartAndEnd();
+
+		// If true, this variable has additional parameters
+		if (CommandKey.Contains(" "))
+		{
+			const int32 IndexOfFirstSpace = CommandString.Find(" ", ESearchCase::IgnoreCase, ESearchDir::FromStart);
+			AdditionalParams = CommandKey.RightChop(IndexOfFirstSpace).TrimStartAndEnd();
+			CommandKey = CommandKey.Left(IndexOfFirstSpace).TrimStartAndEnd();
+		}
 	}
 	else
 	{
@@ -105,17 +126,11 @@ FReply FConsoleVariablesEditorMainPanel::ValidateConsoleInputAndAddToCurrentPres
 		}
 		
 		AddConsoleObjectToCurrentPreset(
-			CommandKey,
+			CommandKey + " " + AdditionalParams,
 			ValueString,
 			true
 		);
 	}
-	else if (CommandString.IsEmpty())
-	{
-		UE_LOG(LogConsoleVariablesEditor, Warning, TEXT("%hs: Input is blank."), __FUNCTION__);
-
-		return FReply::Unhandled();
-	}
 	// Try to execute the whole command. Some commands are not registered, but are parsed externally so they won't be found by IConsoleManager::FindConsoleObject
 	else if (GEngine->Exec(FConsoleVariablesEditorCommandInfo::GetCurrentWorld(), *CommandString))
 	{
