Console Variables Editor:

-Change out single line if-init statements for nested if statements
-Fix for: Insufficient Multiuser Button Selection Padding
-Fix for: Console Variables Editor is not setting edited cvars as Session source

diff --git a/Engine/Plugins/Editor/ConsoleVariablesEditor/Source/ConsoleVariablesEditor/Private/ConsoleVariablesEditorCommandInfo.h b/Engine/Plugins/Editor/ConsoleVariablesEditor/Source/ConsoleVariablesEditor/Private/ConsoleVariablesEditorCommandInfo.h
--- a/Engine/Plugins/Editor/ConsoleVariablesEditor/Source/ConsoleVariablesEditor/Private/ConsoleVariablesEditorCommandInfo.h
+++ b/Engine/Plugins/Editor/ConsoleVariablesEditor/Source/ConsoleVariablesEditor/Private/ConsoleVariablesEditorCommandInfo.h
@@ -64,14 +64,20 @@ struct FConsoleVariablesEditorCommandInfo
 
 	void SetIfChangedInCurrentPreset(const bool bNewSetting)
 	{
-		bChangedInCurrentPreset = bNewSetting;
+		bSetInCurrentSession = bNewSetting;
 	}
 
-	void ExecuteCommand(const FString& NewValueAsString)
+	/** Sets a variable to the specified value whilst maintaining its SetBy flag.
+	  * Non-variables will be executed through the console.
+	  * If bSetInSession is true, this CommandInfo's associated variable row will display "Session" in the UI.
+	 */
+	void ExecuteCommand(const FString& NewValueAsString, const bool bSetInSession = true)
 	{
 		if (IConsoleVariable* AsVariable = GetConsoleVariablePtr())
 		{
 			AsVariable->Set(*NewValueAsString, GetSource());
+
+			bSetInCurrentSession = bSetInSession;
 		}
 		else
 		{
@@ -211,7 +217,7 @@ struct FConsoleVariablesEditorCommandInfo
 			return LOCTEXT("Source_IsNotConsoleVariableButConsoleCommand", "Command");
 		}
 
-		if (bChangedInCurrentPreset)
+		if (bSetInCurrentSession)
 		{
 			return LOCTEXT("Source_SetByCurrentPreset", "Session");
 		}
@@ -278,7 +284,7 @@ struct FConsoleVariablesEditorCommandInfo
 	EConsoleVariableFlags StartupSource = ECVF_Default;
 
 	/** If the variable was last changed by the current preset */
-	bool bChangedInCurrentPreset = false;
+	bool bSetInCurrentSession = false;
 
 	/** When variables change, this callback is executed. */
 	FDelegateHandle OnVariableChangedCallbackHandle;
diff --git a/Engine/Plugins/Editor/ConsoleVariablesEditor/Source/ConsoleVariablesEditor/Private/Views/List/ConsoleVariablesEditorListFilters/ConsoleVariablesEditorListFilter_SetByCurrentPreset.h b/Engine/Plugins/Editor/ConsoleVariablesEditor/Source/ConsoleVariablesEditor/Private/Views/List/ConsoleVariablesEditorListFilters/ConsoleVariablesEditorListFilter_SetByCurrentPreset.h
--- a/Engine/Plugins/Editor/ConsoleVariablesEditor/Source/ConsoleVariablesEditor/Private/Views/List/ConsoleVariablesEditorListFilters/ConsoleVariablesEditorListFilter_SetByCurrentPreset.h
+++ b/Engine/Plugins/Editor/ConsoleVariablesEditor/Source/ConsoleVariablesEditor/Private/Views/List/ConsoleVariablesEditorListFilters/ConsoleVariablesEditorListFilter_SetByCurrentPreset.h
@@ -34,9 +34,12 @@ public:
 
 	virtual bool DoesItemPassFilter(const FConsoleVariablesEditorListRowPtr& InItem) override
 	{
-		if (InItem.IsValid(); const TSharedPtr<FConsoleVariablesEditorCommandInfo> PinnedCommand = InItem->GetCommandInfo().Pin())
+		if (InItem.IsValid())
 		{
-			return PinnedCommand->bChangedInCurrentPreset;
+			if (const TSharedPtr<FConsoleVariablesEditorCommandInfo> PinnedCommand = InItem->GetCommandInfo().Pin())
+			{
+				return PinnedCommand->bSetInCurrentSession;
+			}
 		}
 
 		return false;
diff --git a/Engine/Plugins/Editor/ConsoleVariablesEditor/Source/ConsoleVariablesEditor/Private/Views/List/ConsoleVariablesEditorListFilters/ConsoleVariablesEditorListFilter_SetInSession.h b/Engine/Plugins/Editor/ConsoleVariablesEditor/Source/ConsoleVariablesEditor/Private/Views/List/ConsoleVariablesEditorListFilters/ConsoleVariablesEditorListFilter_SetInSession.h
--- a/Engine/Plugins/Editor/ConsoleVariablesEditor/Source/ConsoleVariablesEditor/Private/Views/List/ConsoleVariablesEditorListFilters/ConsoleVariablesEditorListFilter_SetInSession.h
+++ b/Engine/Plugins/Editor/ConsoleVariablesEditor/Source/ConsoleVariablesEditor/Private/Views/List/ConsoleVariablesEditorListFilters/ConsoleVariablesEditorListFilter_SetInSession.h
@@ -34,9 +34,12 @@ public:
 
 	virtual bool DoesItemPassFilter(const FConsoleVariablesEditorListRowPtr& InItem) override
 	{
-		if (InItem.IsValid(); const TSharedPtr<FConsoleVariablesEditorCommandInfo> PinnedCommand = InItem->GetCommandInfo().Pin())
+		if (InItem.IsValid())
 		{
-			return PinnedCommand->bChangedInCurrentPreset;
+			if (const TSharedPtr<FConsoleVariablesEditorCommandInfo> PinnedCommand = InItem->GetCommandInfo().Pin())
+			{
+				return PinnedCommand->bSetInCurrentSession;
+			}
 		}
 
 		return false;
diff --git a/Engine/Plugins/Editor/ConsoleVariablesEditor/Source/ConsoleVariablesEditor/Private/Views/List/ConsoleVariablesEditorListFilters/ConsoleVariablesEditorListFilter_ShowOnlyCommands.h b/Engine/Plugins/Editor/ConsoleVariablesEditor/Source/ConsoleVariablesEditor/Private/Views/List/ConsoleVariablesEditorListFilters/ConsoleVariablesEditorListFilter_ShowOnlyCommands.h
--- a/Engine/Plugins/Editor/ConsoleVariablesEditor/Source/ConsoleVariablesEditor/Private/Views/List/ConsoleVariablesEditorListFilters/ConsoleVariablesEditorListFilter_ShowOnlyCommands.h
+++ b/Engine/Plugins/Editor/ConsoleVariablesEditor/Source/ConsoleVariablesEditor/Private/Views/List/ConsoleVariablesEditorListFilters/ConsoleVariablesEditorListFilter_ShowOnlyCommands.h
@@ -34,9 +34,12 @@ public:
 
 	virtual bool DoesItemPassFilter(const FConsoleVariablesEditorListRowPtr& InItem) override
 	{
-		if (InItem.IsValid(); const TSharedPtr<FConsoleVariablesEditorCommandInfo> PinnedCommand = InItem->GetCommandInfo().Pin())
+		if (InItem.IsValid())
 		{
-			return PinnedCommand->ObjectType != FConsoleVariablesEditorCommandInfo::EConsoleObjectType::Variable;
+			if (const TSharedPtr<FConsoleVariablesEditorCommandInfo> PinnedCommand = InItem->GetCommandInfo().Pin())
+			{
+				return PinnedCommand->ObjectType != FConsoleVariablesEditorCommandInfo::EConsoleObjectType::Variable;
+			}
 		}
 
 		return false;
diff --git a/Engine/Plugins/Editor/ConsoleVariablesEditor/Source/ConsoleVariablesEditor/Private/Views/List/ConsoleVariablesEditorListFilters/ConsoleVariablesEditorListFilter_ShowOnlyVariables.h b/Engine/Plugins/Editor/ConsoleVariablesEditor/Source/ConsoleVariablesEditor/Private/Views/List/ConsoleVariablesEditorListFilters/ConsoleVariablesEditorListFilter_ShowOnlyVariables.h
--- a/Engine/Plugins/Editor/ConsoleVariablesEditor/Source/ConsoleVariablesEditor/Private/Views/List/ConsoleVariablesEditorListFilters/ConsoleVariablesEditorListFilter_ShowOnlyVariables.h
+++ b/Engine/Plugins/Editor/ConsoleVariablesEditor/Source/ConsoleVariablesEditor/Private/Views/List/ConsoleVariablesEditorListFilters/ConsoleVariablesEditorListFilter_ShowOnlyVariables.h
@@ -34,9 +34,12 @@ public:
 
 	virtual bool DoesItemPassFilter(const FConsoleVariablesEditorListRowPtr& InItem) override
 	{
-		if (InItem.IsValid(); const TSharedPtr<FConsoleVariablesEditorCommandInfo> PinnedCommand = InItem->GetCommandInfo().Pin())
+		if (InItem.IsValid())
 		{
-			return PinnedCommand->ObjectType == FConsoleVariablesEditorCommandInfo::EConsoleObjectType::Variable;
+			if (const TSharedPtr<FConsoleVariablesEditorCommandInfo> PinnedCommand = InItem->GetCommandInfo().Pin())
+			{
+				return PinnedCommand->ObjectType == FConsoleVariablesEditorCommandInfo::EConsoleObjectType::Variable;
+			}
 		}
 
 		return false;
diff --git a/Engine/Plugins/Editor/ConsoleVariablesEditor/Source/ConsoleVariablesEditor/Private/Views/List/ConsoleVariablesEditorListFilters/ConsoleVariablesEditorListFilter_Source.h b/Engine/Plugins/Editor/ConsoleVariablesEditor/Source/ConsoleVariablesEditor/Private/Views/List/ConsoleVariablesEditorListFilters/ConsoleVariablesEditorListFilter_Source.h
--- a/Engine/Plugins/Editor/ConsoleVariablesEditor/Source/ConsoleVariablesEditor/Private/Views/List/ConsoleVariablesEditorListFilters/ConsoleVariablesEditorListFilter_Source.h
+++ b/Engine/Plugins/Editor/ConsoleVariablesEditor/Source/ConsoleVariablesEditor/Private/Views/List/ConsoleVariablesEditorListFilters/ConsoleVariablesEditorListFilter_Source.h
@@ -32,10 +32,13 @@ public:
 
 	virtual bool DoesItemPassFilter(const FConsoleVariablesEditorListRowPtr& InItem) override
 	{
-		if (InItem.IsValid(); const TSharedPtr<FConsoleVariablesEditorCommandInfo> PinnedCommand = InItem->GetCommandInfo().Pin())
+		if (InItem.IsValid())
 		{
-			const EConsoleVariableFlags Source = PinnedCommand->GetSource();
-			return Source == SourceFlag;
+			if (const TSharedPtr<FConsoleVariablesEditorCommandInfo> PinnedCommand = InItem->GetCommandInfo().Pin())
+			{
+				const EConsoleVariableFlags Source = PinnedCommand->GetSource();
+				return Source == SourceFlag;
+			}
 		}
 
 		return false;
diff --git a/Engine/Plugins/Editor/ConsoleVariablesEditor/Source/ConsoleVariablesEditor/Private/Views/List/ConsoleVariablesEditorListRow.cpp b/Engine/Plugins/Editor/ConsoleVariablesEditor/Source/ConsoleVariablesEditor/Private/Views/List/ConsoleVariablesEditorListRow.cpp
--- a/Engine/Plugins/Editor/ConsoleVariablesEditor/Source/ConsoleVariablesEditor/Private/Views/List/ConsoleVariablesEditorListRow.cpp
+++ b/Engine/Plugins/Editor/ConsoleVariablesEditor/Source/ConsoleVariablesEditor/Private/Views/List/ConsoleVariablesEditorListRow.cpp
@@ -125,8 +125,7 @@ void FConsoleVariablesEditorListRow::ResetToStartupValueAndSource() const
 	if (const TSharedPtr<FConsoleVariablesEditorCommandInfo> PinnedCommand = GetCommandInfo().Pin())
 	{
 		PinnedCommand->SetSourceFlag(PinnedCommand->StartupSource);
-		PinnedCommand->ExecuteCommand(PinnedCommand->StartupValueAsString);
-		PinnedCommand->SetIfChangedInCurrentPreset(false);
+		PinnedCommand->ExecuteCommand(PinnedCommand->StartupValueAsString, false);
 	}
 }
 
diff --git a/Engine/Plugins/Editor/ConsoleVariablesEditor/Source/ConsoleVariablesEditor/Private/Views/List/SConsoleVariablesEditorList.cpp b/Engine/Plugins/Editor/ConsoleVariablesEditor/Source/ConsoleVariablesEditor/Private/Views/List/SConsoleVariablesEditorList.cpp
--- a/Engine/Plugins/Editor/ConsoleVariablesEditor/Source/ConsoleVariablesEditor/Private/Views/List/SConsoleVariablesEditorList.cpp
+++ b/Engine/Plugins/Editor/ConsoleVariablesEditor/Source/ConsoleVariablesEditor/Private/Views/List/SConsoleVariablesEditorList.cpp
@@ -120,7 +120,7 @@ void SConsoleVariablesEditorList::Construct(const FArguments& InArgs, TSharedRef
 			.Padding(8.f, 1.f, 0.f, 1.f)
 			[
 				SAssignNew( ViewOptionsComboButton, SComboButton )
-				.ContentPadding(0)
+				.ContentPadding(4.f)
 				.ToolTipText(LOCTEXT("ShowOptions_Tooltip", "Show options to affect the visibility of items in the Console Variables Editor list"))
 				.ComboButtonStyle( FAppStyle::Get(), "SimpleComboButtonWithIcon" ) // Use the tool bar item style for this button
 				.OnGetMenuContent( this, &SConsoleVariablesEditorList::BuildShowOptionsMenu)
diff --git a/Engine/Plugins/Editor/ConsoleVariablesEditor/Source/ConsoleVariablesEditor/Private/Views/MainPanel/SConsoleVariablesEditorMainPanel.cpp b/Engine/Plugins/Editor/ConsoleVariablesEditor/Source/ConsoleVariablesEditor/Private/Views/MainPanel/SConsoleVariablesEditorMainPanel.cpp
--- a/Engine/Plugins/Editor/ConsoleVariablesEditor/Source/ConsoleVariablesEditor/Private/Views/MainPanel/SConsoleVariablesEditorMainPanel.cpp
+++ b/Engine/Plugins/Editor/ConsoleVariablesEditor/Source/ConsoleVariablesEditor/Private/Views/MainPanel/SConsoleVariablesEditorMainPanel.cpp
@@ -199,7 +199,7 @@ void SConsoleVariablesEditorMainPanel::CreateConcertButtonIfNeeded()
 		.Padding(11.f, 1.f, 1.f, 1.f)
 		[
 			SAssignNew(ConcertButtonPtr, SCheckBox)
-			.Padding(0)
+			.Padding(4.f)
 			.ToolTipText(LOCTEXT("ShowConcertSettings_Tip", "Show the multi-user controls for Console Variables"))
 			.Style(&FAppStyle::Get().GetWidgetStyle<FCheckBoxStyle>("ToggleButtonCheckbox"))
 			.ForegroundColor(FSlateColor::UseForeground())
