Console Variables Editor:

-Fix for: Ensure and multiple entries when setting r.Shadow.Virtual.Enable 0
-Fix for: Variables not transacting when set from blueprint function library or when unchecking a row in CVE

diff --git a/Engine/Plugins/Editor/ConsoleVariablesEditor/Source/ConsoleVariablesEditor/Private/ConsoleVariablesAsset.cpp b/Engine/Plugins/Editor/ConsoleVariablesEditor/Source/ConsoleVariablesEditor/Private/ConsoleVariablesAsset.cpp
--- a/Engine/Plugins/Editor/ConsoleVariablesEditor/Source/ConsoleVariablesEditor/Private/ConsoleVariablesAsset.cpp
+++ b/Engine/Plugins/Editor/ConsoleVariablesEditor/Source/ConsoleVariablesEditor/Private/ConsoleVariablesAsset.cpp
@@ -2,8 +2,9 @@
 
 #include "ConsoleVariablesAsset.h"
 
+#include "ConsoleVariablesEditorCommandInfo.h"
+#include "ConsoleVariablesEditorLog.h"
 #include "ConsoleVariablesEditorModule.h"
-#include "Algo/Find.h"
 
 void UConsoleVariablesAsset::SetVariableCollectionDescription(const FString& InVariableCollectionDescription)
 {
@@ -15,11 +16,12 @@ void UConsoleVariablesAsset::ReplaceSavedCommands(const TArray<FConsoleVariables
 	SavedCommands = Replacement;
 }
 
-bool UConsoleVariablesAsset::FindSavedDataByCommandString(const FString InCommandString, FConsoleVariablesEditorAssetSaveData& OutValue) const
+bool UConsoleVariablesAsset::FindSavedDataByCommandString(
+	const FString InCommandString, FConsoleVariablesEditorAssetSaveData& OutValue, const ESearchCase::Type SearchCase) const
 {
 	for (const FConsoleVariablesEditorAssetSaveData& Command : SavedCommands)
 	{
-		if (Command.CommandName.Equals(InCommandString))
+		if (Command.CommandName.TrimStartAndEnd().Equals(InCommandString.TrimStartAndEnd(), SearchCase))
 		{
 			OutValue = Command;
 			return true;
@@ -32,7 +34,8 @@ bool UConsoleVariablesAsset::FindSavedDataByCommandString(const FString InComman
 void UConsoleVariablesAsset::AddOrSetConsoleObjectSavedData(const FConsoleVariablesEditorAssetSaveData& InData)
 {
 	FConsoleVariablesEditorModule& ConsoleVariablesEditorModule = FConsoleVariablesEditorModule::Get();
-
+	UE_LOG(LogConsoleVariablesEditor, VeryVerbose, TEXT("%hs: Adding %s to editable asset"),
+		__FUNCTION__, *InData.CommandName);
 	if (const TWeakPtr<FConsoleVariablesEditorCommandInfo> CommandInfo =
 			ConsoleVariablesEditorModule.FindCommandInfoByName(InData.CommandName);
 		CommandInfo.IsValid())
@@ -50,17 +53,32 @@ void UConsoleVariablesAsset::AddOrSetConsoleObjectSavedData(const FConsoleVariab
 	
 	RemoveConsoleVariable(InData.CommandName);
 	SavedCommands.Add(InData);
+	// Make a copy because this array can change during iteration
+	TArray<FConsoleVariablesEditorAssetSaveData> SavedCommandsLocal = SavedCommands; 
+	for (int32 CommandIndex = 0; CommandIndex < SavedCommandsLocal.Num(); CommandIndex++)
+	{
+		FConsoleVariablesEditorAssetSaveData& SavedCommand = SavedCommandsLocal[CommandIndex];
+		UE_LOG(LogConsoleVariablesEditor, VeryVerbose, TEXT("%hs: Command named '%s' at Index %i"),
+			__FUNCTION__, *SavedCommand.CommandName, CommandIndex);
+	}
 }
 
 bool UConsoleVariablesAsset::RemoveConsoleVariable(const FString InCommandString)
 {
 	FConsoleVariablesEditorAssetSaveData ExistingData;
-	if (FindSavedDataByCommandString(InCommandString, ExistingData))
+	int32 RemoveCount = 0;
+	while (FindSavedDataByCommandString(InCommandString, ExistingData, ESearchCase::IgnoreCase))
 	{
-		return SavedCommands.Remove(ExistingData) > 0;
+		UE_LOG(LogConsoleVariablesEditor, VeryVerbose, TEXT("%hs: Removing %s from editable asset"),
+			__FUNCTION__, *InCommandString);
+		
+		if (SavedCommands.Remove(ExistingData) > 0)
+		{
+			RemoveCount++;
+		}
 	}
 	
-	return false;
+	return RemoveCount > 0;
 }
 
 void UConsoleVariablesAsset::CopyFrom(const UConsoleVariablesAsset* InAssetToCopy)
diff --git a/Engine/Plugins/Editor/ConsoleVariablesEditor/Source/ConsoleVariablesEditor/Private/ConsoleVariablesEditorCommandInfo.cpp b/Engine/Plugins/Editor/ConsoleVariablesEditor/Source/ConsoleVariablesEditor/Private/ConsoleVariablesEditorCommandInfo.cpp
new file mode 100644
--- /dev/null
+++ b/Engine/Plugins/Editor/ConsoleVariablesEditor/Source/ConsoleVariablesEditor/Private/ConsoleVariablesEditorCommandInfo.cpp
@@ -0,0 +1,204 @@
+// Copyright Epic Games, Inc. All Rights Reserved.
+
+#pragma once
+
+#include "ConsoleVariablesEditorCommandInfo.h"
+#include "Algo/Find.h"
+#include "CoreMinimal.h"
+#include "ConsoleVariablesEditorLog.h"
+#include "Editor.h"
+#include "Editor/EditorEngine.h"
+#include "Engine/GameEngine.h"
+#include "HAL/IConsoleManager.h"
+
+#define LOCTEXT_NAMESPACE "ConsoleVariablesEditor"
+
+void FConsoleVariablesEditorCommandInfo::ExecuteCommand(
+	const FString& NewValueAsString, const bool bShouldTransactInConcert, const bool bSetInSession)
+{
+	if (IConsoleVariable* AsVariable = GetConsoleVariablePtr())
+	{
+		AsVariable->Set(*NewValueAsString, GetSource());
+		bSetInCurrentSession = bSetInSession;
+	}
+	else
+	{
+		GEngine->Exec(GetCurrentWorld(),
+			*FString::Printf(TEXT("%s %s"), *Command, *NewValueAsString).TrimStartAndEnd());
+	}
+	if (bShouldTransactInConcert)
+	{
+		FConsoleVariablesEditorModule::Get().SendMultiUserConsoleVariableChange(Command, NewValueAsString);
+	}
+}
+
+/** Get a reference to the cached console object. May return nullptr if unregistered. */
+IConsoleObject* FConsoleVariablesEditorCommandInfo::GetConsoleObjectPtr()
+{
+	// If the console object ptr goes stale or is older than the specified threshold, try to refresh it
+	// May return nullptr if unregistered
+	if (!ConsoleObjectPtr ||
+		(FDateTime::UtcNow() - TimeOfLastConsoleObjectRefresh).GetTotalSeconds() > ConsoleObjectRefreshThreshold)
+	{
+		FString CommandKey = Command; 
+		// Remove additional params, if they exist
+		const int32 IndexOfSpace = CommandKey.Find(" ");
+		if (IndexOfSpace != INDEX_NONE)
+		{
+			CommandKey = CommandKey.Left(IndexOfSpace).TrimStartAndEnd();
+		}
+		
+		ConsoleObjectPtr = IConsoleManager::Get().FindConsoleObject(*CommandKey);
+		TimeOfLastConsoleObjectRefresh = FDateTime::UtcNow();
+	}
+	// If the console object turns out to be unregistered, let interested parties know
+	if (ConsoleObjectPtr && ConsoleObjectPtr->TestFlags(ECVF_Unregistered))
+	{
+		OnDetectConsoleObjectUnregistered.Broadcast(Command);
+	}
+	return ConsoleObjectPtr;
+}
+
+/** Return the console object as a console variable object if applicable. May return nullptr if unregistered. */
+IConsoleVariable* FConsoleVariablesEditorCommandInfo::GetConsoleVariablePtr()
+{
+	if (IConsoleObject* ObjectPtr = GetConsoleObjectPtr())
+	{
+		return ObjectPtr->AsVariable();
+	}
+	
+	return nullptr;
+}
+
+/**
+ *Return the console object as a console command object if applicable.
+ *Does not consider externally parsed console commands, as they have no associated objects.
+ */
+IConsoleCommand* FConsoleVariablesEditorCommandInfo::GetConsoleCommandPtr()
+{
+	if (IConsoleObject* ObjectPtr = GetConsoleObjectPtr())
+	{
+		return ObjectPtr->AsCommand();
+	}
+	
+	return nullptr;
+}
+
+UWorld* FConsoleVariablesEditorCommandInfo::GetCurrentWorld()
+{
+	UWorld* CurrentWorld = nullptr;
+	if (GIsEditor)
+	{
+		CurrentWorld = GEditor->GetEditorWorldContext().World();
+	}
+	else if (UGameEngine* GameEngine = Cast<UGameEngine>(GEngine))
+	{
+		CurrentWorld = GameEngine->GetGameWorld();
+	}
+	return CurrentWorld;
+}
+
+FString FConsoleVariablesEditorCommandInfo::GetHelpText()
+{
+	if (const IConsoleVariable* AsVariable = GetConsoleVariablePtr())
+	{
+		return FString(AsVariable->GetHelp());
+	}
+	return "";
+}
+
+EConsoleVariableFlags FConsoleVariablesEditorCommandInfo::GetSource()
+{
+	if (const IConsoleObject* ConsoleObject = GetConsoleObjectPtr())
+	{
+		return (EConsoleVariableFlags)((uint32)ConsoleObject->GetFlags() & ECVF_SetByMask);
+	}
+	return ECVF_Default;
+}
+
+void FConsoleVariablesEditorCommandInfo::ClearSourceFlags()
+{
+	if (IConsoleObject* ConsoleObject = GetConsoleObjectPtr())
+	{
+		for (const FStaticConsoleVariableFlagInfo& StaticConsoleVariableFlagInfo : SupportedFlags)
+		{
+			ConsoleObject->ClearFlags(StaticConsoleVariableFlagInfo.Flag);
+		}
+	}
+}
+
+void FConsoleVariablesEditorCommandInfo::SetSourceFlag(const EConsoleVariableFlags InSource)
+{
+	if (IConsoleVariable* AsVariable = GetConsoleVariablePtr())
+	{
+		AsVariable->Set(*AsVariable->GetString(), StartupSource);
+		return;
+	}
+	
+	const uint32 OldPri = (uint32)GetSource();
+	const uint32 NewPri = (uint32)InSource;
+	if (NewPri < OldPri)
+	{
+		return;
+	}
+	
+	if (IConsoleObject* ConsoleObject = GetConsoleObjectPtr())
+	{
+		ClearSourceFlags();
+		ConsoleObject->SetFlags((EConsoleVariableFlags)InSource);
+	}
+}
+
+FText FConsoleVariablesEditorCommandInfo::GetSourceAsText()
+{
+	// Non-variables don't really have a source
+	if (ObjectType != EConsoleObjectType::Variable)
+	{
+		return LOCTEXT("Source_IsNotConsoleVariableButConsoleCommand", "Command");
+	}
+	if (bSetInCurrentSession)
+	{
+		return LOCTEXT("Source_SetByCurrentPreset", "Session");
+	}
+	
+	return ConvertConsoleVariableSetByFlagToText(GetSource());
+}
+
+FText FConsoleVariablesEditorCommandInfo::ConvertConsoleVariableSetByFlagToText(const EConsoleVariableFlags InFlag)
+{
+	FText ReturnValue = LOCTEXT("UnknownSource", "<UNKNOWN>"); 
+	if (const FStaticConsoleVariableFlagInfo* Match = Algo::FindByPredicate(
+		SupportedFlags,
+			[InFlag](const FStaticConsoleVariableFlagInfo& Comparator)
+			{
+				return Comparator.Flag == InFlag;
+			}))
+	{
+		ReturnValue = (*Match).DisplayText;
+	}
+	
+	return ReturnValue;
+}
+
+bool FConsoleVariablesEditorCommandInfo::IsCurrentValueDifferentFromInputValue(const FString& InValueToCompare)
+{
+	if (const IConsoleVariable* AsVariable = GetConsoleVariablePtr())
+	{
+		// Floats sometimes return true erroneously because they can be stringified as e.g '1' or '1.0' by different functions.
+		if (AsVariable->IsVariableFloat())
+		{
+			const float A = AsVariable->GetFloat();
+			const float B = FCString::Atof(*InValueToCompare);
+			return !FMath::IsNearlyEqual(A, B);
+		}
+		
+		return !AsVariable->GetString().Equals(InValueToCompare);
+	}
+	else if (ObjectType == EConsoleObjectType::NullObject || ObjectType == EConsoleObjectType::Command)
+	{
+		return true;
+	}
+	return false;
+}
+
+#undef LOCTEXT_NAMESPACE
diff --git a/Engine/Plugins/Editor/ConsoleVariablesEditor/Source/ConsoleVariablesEditor/Private/ConsoleVariablesEditorCommandInfo.h b/Engine/Plugins/Editor/ConsoleVariablesEditor/Source/ConsoleVariablesEditor/Private/ConsoleVariablesEditorCommandInfo.h
--- a/Engine/Plugins/Editor/ConsoleVariablesEditor/Source/ConsoleVariablesEditor/Private/ConsoleVariablesEditorCommandInfo.h
+++ b/Engine/Plugins/Editor/ConsoleVariablesEditor/Source/ConsoleVariablesEditor/Private/ConsoleVariablesEditorCommandInfo.h
@@ -1,13 +1,9 @@
 // Copyright Epic Games, Inc. All Rights Reserved.
 
 #pragma once
-
-#include "Algo/Find.h"
+ 
 #include "CoreMinimal.h"
-#include "ConsoleVariablesEditorLog.h"
-#include "Editor.h"
-#include "Engine/GameEngine.h"
-#include "HAL/IConsoleManager.h"
+#include "ConsoleVariablesEditorModule.h"
 
 DECLARE_MULTICAST_DELEGATE_OneParam(FOnDetectConsoleObjectUnregistered, FString)
 
@@ -24,13 +20,13 @@ struct FConsoleVariablesEditorCommandInfo
 		// A console variable such as 'r.ScreenPercentage'
 		Variable
 	};
-
+	
 	struct FStaticConsoleVariableFlagInfo
 	{
 		EConsoleVariableFlags Flag;
 		FText DisplayText;
 	};
-
+	
 	FConsoleVariablesEditorCommandInfo(const FString& InCommand)
 	: Command(InCommand)
 	{
@@ -46,235 +42,68 @@ struct FConsoleVariablesEditorCommandInfo
 			}
 		}
 	}
-
+	
 	~FConsoleVariablesEditorCommandInfo()
 	{
 		OnDetectConsoleObjectUnregistered.Remove(OnDetectConsoleObjectUnregisteredHandle);
-
 		if (IConsoleVariable* AsVariable = GetConsoleVariablePtr())
 		{
 			AsVariable->OnChangedDelegate().Remove(OnVariableChangedCallbackHandle);
 		}
 	}
-
+	
 	FORCEINLINE bool operator==(const FConsoleVariablesEditorCommandInfo& Comparator) const
 	{
-		return Command.Equals(Comparator.Command);
+		return Command.Equals(Comparator.Command, ESearchCase::IgnoreCase);
 	}
-
+	
 	void SetIfChangedInCurrentPreset(const bool bNewSetting)
 	{
 		bSetInCurrentSession = bNewSetting;
 	}
-
+	
 	/** Sets a variable to the specified value whilst maintaining its SetBy flag.
 	  * Non-variables will be executed through the console.
 	  * If bSetInSession is true, this CommandInfo's associated variable row will display "Session" in the UI.
 	 */
-	void ExecuteCommand(const FString& NewValueAsString, const bool bSetInSession = true)
-	{
-		if (IConsoleVariable* AsVariable = GetConsoleVariablePtr())
-		{
-			AsVariable->Set(*NewValueAsString, GetSource());
-
-			bSetInCurrentSession = bSetInSession;
-		}
-		else
-		{
-			GEngine->Exec(GetCurrentWorld(),
-				*FString::Printf(TEXT("%s %s"), *Command, *NewValueAsString).TrimStartAndEnd());
-		}
-	}
-
+	void ExecuteCommand(
+		const FString& NewValueAsString, const bool bShouldTransactInConcert = true, const bool bSetInSession = true);
+ 
 	/** Get a reference to the cached console object. May return nullptr if unregistered. */
-	IConsoleObject* GetConsoleObjectPtr()
-	{
-		// If the console object ptr goes stale or is older than the specified threshold, try to refresh it
-		// May return nullptr if unregistered
-		if (!ConsoleObjectPtr ||
-			(FDateTime::UtcNow() - TimeOfLastConsoleObjectRefresh).GetTotalSeconds() > ConsoleObjectRefreshThreshold)
-		{
-			FString CommandKey = Command; 
-
-			// Remove additional params, if they exist
-			const int32 IndexOfSpace = CommandKey.Find(" ");
-			if (IndexOfSpace != INDEX_NONE)
-			{
-				CommandKey = CommandKey.Left(IndexOfSpace).TrimStartAndEnd();
-			}
-			
-			ConsoleObjectPtr = IConsoleManager::Get().FindConsoleObject(*CommandKey);
-			TimeOfLastConsoleObjectRefresh = FDateTime::UtcNow();
-		}
-
-		// If the console object turns out to be unregistered, let interested parties know
-		if (ConsoleObjectPtr && ConsoleObjectPtr->TestFlags(ECVF_Unregistered))
-		{
-			OnDetectConsoleObjectUnregistered.Broadcast(Command);
-		}
-
-		return ConsoleObjectPtr;
-	}
+	IConsoleObject* GetConsoleObjectPtr();
 
 	/** Return the console object as a console variable object if applicable. May return nullptr if unregistered. */
-	IConsoleVariable* GetConsoleVariablePtr()
-	{
-		if (IConsoleObject* ObjectPtr = GetConsoleObjectPtr())
-		{
-			return ObjectPtr->AsVariable();
-		}
-		
-		return nullptr;
-	}
+	IConsoleVariable* GetConsoleVariablePtr();
 
 	/**
 	 *Return the console object as a console command object if applicable.
 	 *Does not consider externally parsed console commands, as they have no associated objects.
 	 */
-	IConsoleCommand* GetConsoleCommandPtr()
-	{
-		if (IConsoleObject* ObjectPtr = GetConsoleObjectPtr())
-		{
-			return ObjectPtr->AsCommand();
-		}
-		
-		return nullptr;
-	}
-	
-	static UWorld* GetCurrentWorld()
-	{
-		UWorld* CurrentWorld = nullptr;
-		if (GIsEditor)
-		{
-			CurrentWorld = GEditor->GetEditorWorldContext().World();
-		}
-		else if (UGameEngine* GameEngine = Cast<UGameEngine>(GEngine))
-		{
-			CurrentWorld = GameEngine->GetGameWorld();
-		}
-		return CurrentWorld;
-	}
-
-	FString GetHelpText()
-	{
-		if (const IConsoleVariable* AsVariable = GetConsoleVariablePtr())
-		{
-			return FString(AsVariable->GetHelp());
-		}
-
-		return "";
-	}
-
-	EConsoleVariableFlags GetSource()
-	{
-		if (const IConsoleObject* ConsoleObject = GetConsoleObjectPtr())
-		{
-			return (EConsoleVariableFlags)((uint32)ConsoleObject->GetFlags() & ECVF_SetByMask);
-		}
-
-		return ECVF_Default;
-	}
-
-	void ClearSourceFlags()
-	{
-		if (IConsoleObject* ConsoleObject = GetConsoleObjectPtr())
-		{
-			for (const FStaticConsoleVariableFlagInfo& StaticConsoleVariableFlagInfo : SupportedFlags)
-			{
-				ConsoleObject->ClearFlags(StaticConsoleVariableFlagInfo.Flag);
-			}
-		}
-	}
-
-	void SetSourceFlag(const EConsoleVariableFlags InSource)
-	{
-		if (IConsoleVariable* AsVariable = GetConsoleVariablePtr())
-		{
-			AsVariable->Set(*AsVariable->GetString(), StartupSource);
-			return;
-		}
-		
-		const uint32 OldPri = (uint32)GetSource();
-		const uint32 NewPri = (uint32)InSource;
-
-		if (NewPri < OldPri)
-		{
-			return;
-		}
-		
-		if (IConsoleObject* ConsoleObject = GetConsoleObjectPtr())
-		{
-			ClearSourceFlags();
-			ConsoleObject->SetFlags((EConsoleVariableFlags)InSource);
-		}
-	}
-
-	FText GetSourceAsText()
-	{
-		// Non-variables don't really have a source
-		if (ObjectType != EConsoleObjectType::Variable)
-		{
-			return LOCTEXT("Source_IsNotConsoleVariableButConsoleCommand", "Command");
-		}
-
-		if (bSetInCurrentSession)
-		{
-			return LOCTEXT("Source_SetByCurrentPreset", "Session");
-		}
-		
-		return ConvertConsoleVariableSetByFlagToText(GetSource());
-	}
-
-	static FText ConvertConsoleVariableSetByFlagToText(const EConsoleVariableFlags InFlag)
-	{
-		FText ReturnValue = LOCTEXT("UnknownSource", "<UNKNOWN>"); 
-
-		if (const FStaticConsoleVariableFlagInfo* Match = Algo::FindByPredicate(
-			SupportedFlags,
-				[InFlag](const FStaticConsoleVariableFlagInfo& Comparator)
-				{
-					return Comparator.Flag == InFlag;
-				}))
-		{
-			ReturnValue = (*Match).DisplayText;
-		}
-		
-		return ReturnValue;
-	}
-
-	bool IsCurrentValueDifferentFromInputValue(const FString& InValueToCompare)
-	{
-		if (const IConsoleVariable* AsVariable = GetConsoleVariablePtr())
-		{
-			// Floats sometimes return true erroneously because they can be stringified as e.g '1' or '1.0' by different functions.
-			if (AsVariable->IsVariableFloat())
-			{
-				const float A = AsVariable->GetFloat();
-				const float B = FCString::Atof(*InValueToCompare);
-
-				return !FMath::IsNearlyEqual(A, B);
-			}
-			
-			return !AsVariable->GetString().Equals(InValueToCompare);
-		}
-		else if (ObjectType == EConsoleObjectType::NullObject || ObjectType == EConsoleObjectType::Command)
-		{
-			return true;
-		}
-
-		return false;
-	}
-
+	IConsoleCommand* GetConsoleCommandPtr();
+ 
+	static UWorld* GetCurrentWorld();
+
+	FString GetHelpText();
+ 
+	EConsoleVariableFlags GetSource();
+ 
+	void ClearSourceFlags();
+
+	void SetSourceFlag(const EConsoleVariableFlags InSource);
+ 
+	FText GetSourceAsText();
+
+	static FText ConvertConsoleVariableSetByFlagToText(const EConsoleVariableFlags InFlag);
+ 
+	bool IsCurrentValueDifferentFromInputValue(const FString& InValueToCompare);
+ 
 	/** The actual string key or name */
 	UPROPERTY()
 	FString Command;
-
 	EConsoleObjectType ObjectType = EConsoleObjectType::NullObject;
-
 	/** This object is periodically refreshed to mitigate the occurrence of stale pointers. */
 	IConsoleObject* ConsoleObjectPtr;
 	FDateTime TimeOfLastConsoleObjectRefresh;
-
 	double ConsoleObjectRefreshThreshold = 1.0;
 	
 	/** The value of this variable (if Variable object type) when the module started in this session after it may have been set by an ini file. */
@@ -282,17 +111,13 @@ struct FConsoleVariablesEditorCommandInfo
 	
 	/** The source of this variable's (if Variable object type) last setting as recorded when the plugin was loaded. */
 	EConsoleVariableFlags StartupSource = ECVF_Default;
-
 	/** If the variable was last changed by the current preset */
 	bool bSetInCurrentSession = false;
-
 	/** When variables change, this callback is executed. */
 	FDelegateHandle OnVariableChangedCallbackHandle;
-
 	/** When commands are unregistered change, this callback is broadcasted. */
 	FOnDetectConsoleObjectUnregistered OnDetectConsoleObjectUnregistered;
 	FDelegateHandle OnDetectConsoleObjectUnregisteredHandle;
-
 	/** A mapping of SetBy console variable flags to information like the associated display text. */
 	static const inline TArray<FStaticConsoleVariableFlagInfo> SupportedFlags =
 	{
@@ -309,5 +134,4 @@ struct FConsoleVariablesEditorCommandInfo
 		{ EConsoleVariableFlags::ECVF_SetByConsole, LOCTEXT("Source_SetByConsole", "Console") }
 	};
 };
-
 #undef LOCTEXT_NAMESPACE
diff --git a/Engine/Plugins/Editor/ConsoleVariablesEditor/Source/ConsoleVariablesEditor/Private/ConsoleVariablesEditorModule.cpp b/Engine/Plugins/Editor/ConsoleVariablesEditor/Source/ConsoleVariablesEditor/Private/ConsoleVariablesEditorModule.cpp
--- a/Engine/Plugins/Editor/ConsoleVariablesEditor/Source/ConsoleVariablesEditor/Private/ConsoleVariablesEditorModule.cpp
+++ b/Engine/Plugins/Editor/ConsoleVariablesEditor/Source/ConsoleVariablesEditor/Private/ConsoleVariablesEditorModule.cpp
@@ -3,6 +3,7 @@
 #include "ConsoleVariablesEditorModule.h"
 
 #include "AssetTypeActions/AssetTypeActions_ConsoleVariables.h"
+#include "ConsoleVariablesEditorCommandInfo.h"
 #include "ConsoleVariablesEditorLog.h"
 #include "ConsoleVariablesEditorProjectSettings.h"
 #include "ConsoleVariablesEditorStyle.h"
@@ -346,7 +347,7 @@ void FConsoleVariablesEditorModule::OnConsoleVariableChanged(IConsoleVariable* C
 		const FString& Key = PinnedCommand->Command;
 		
 		FConsoleVariablesEditorAssetSaveData FoundData;
-		bool bIsVariableCurrentlyTracked = EditingPresetAsset->FindSavedDataByCommandString(Key, FoundData);
+		bool bIsVariableCurrentlyTracked = EditingPresetAsset->FindSavedDataByCommandString(Key, FoundData, ESearchCase::IgnoreCase);
 
 		const UConsoleVariablesEditorProjectSettings* Settings = GetDefault<UConsoleVariablesEditorProjectSettings>();
 		check(Settings);
diff --git a/Engine/Plugins/Editor/ConsoleVariablesEditor/Source/ConsoleVariablesEditor/Private/ConsoleVariablesEditorModule.h b/Engine/Plugins/Editor/ConsoleVariablesEditor/Source/ConsoleVariablesEditor/Private/ConsoleVariablesEditorModule.h
--- a/Engine/Plugins/Editor/ConsoleVariablesEditor/Source/ConsoleVariablesEditor/Private/ConsoleVariablesEditorModule.h
+++ b/Engine/Plugins/Editor/ConsoleVariablesEditor/Source/ConsoleVariablesEditor/Private/ConsoleVariablesEditorModule.h
@@ -3,7 +3,6 @@
 #pragma once
 
 #include "ConsoleVariablesAsset.h"
-#include "ConsoleVariablesEditorCommandInfo.h"
 #include "MultiUser/ConsoleVariableSync.h"
 
 #include "CoreMinimal.h"
@@ -12,6 +11,7 @@
 #include "Widgets/Docking/SDockTab.h"
 
 struct FAssetData;
+struct FConsoleVariablesEditorCommandInfo;
 
 class FConsoleVariablesEditorMainPanel;
 class FConsoleVariablesEditorToolkit;
diff --git a/Engine/Plugins/Editor/ConsoleVariablesEditor/Source/ConsoleVariablesEditor/Private/MultiUser/ConsoleVariableSync.cpp b/Engine/Plugins/Editor/ConsoleVariablesEditor/Source/ConsoleVariablesEditor/Private/MultiUser/ConsoleVariableSync.cpp
--- a/Engine/Plugins/Editor/ConsoleVariablesEditor/Source/ConsoleVariablesEditor/Private/MultiUser/ConsoleVariableSync.cpp
+++ b/Engine/Plugins/Editor/ConsoleVariablesEditor/Source/ConsoleVariablesEditor/Private/MultiUser/ConsoleVariableSync.cpp
@@ -406,6 +406,22 @@ struct FManagerImpl
 								 EConcertMessageFlags::ReliableOrdered | EConcertMessageFlags::UniqueId);
 	}
 
+	/**
+	 * Sends the given console variable name with the specified checked state to all connected endpoints.  It is up to the
+	 * endpoint to implement the change locally based on configured synchronization state.
+	 */
+	void SendListItemCheckStateChange(FString InName, ECheckBoxState InCheckState)
+	{
+		TSharedPtr<IConcertClientSession> Session = WeakSession.Pin();
+		if (!bIsEnabled || !Session.IsValid())
+		{
+			return;
+		}
+		FConcertSetListItemCheckStateEvent OutEvent{MoveTemp(InName), MoveTemp(InCheckState)};
+		Session->SendCustomEvent(OutEvent, Session->GetSessionClientEndpointIds(),
+								 EConcertMessageFlags::ReliableOrdered | EConcertMessageFlags::UniqueId);
+	}
+
 	/** Reference to the customization object so that we can sync with the session.*/
 	TSharedPtr<FConcertConsoleVariableSessionCustomization> Customization;
 
@@ -421,6 +437,9 @@ struct FManagerImpl
 	/** Multicast delegate invoked when Multi-user console variable change event is received. */
 	FOnRemoteCVarChange RemoteCVarChanged;
 
+	/** Multicast delegate invoked when Multi-user console variable list item in CVE is manually checked or unchecked. */
+	FOnRemoteListItemCheckStateChange RemoteListItemCheckStateChanged;
+
 	/** Multi-user connection change delegate */
 	FOnMultiUserConnectionChange ConnectionChanged;
 
@@ -444,6 +463,11 @@ FOnRemoteCVarChange& FManager::OnRemoteCVarChange()
 	return Implementation->RemoteCVarChanged;
 }
 
+FOnRemoteListItemCheckStateChange& FManager::OnRemoteListItemCheckStateChange()
+{
+	return Implementation->RemoteListItemCheckStateChanged;
+}
+
 FOnMultiUserConnectionChange& FManager::OnConnectionChange()
 {
 	return Implementation->ConnectionChanged;
@@ -454,6 +478,11 @@ void FManager::SendConsoleVariableChange(FString InName, FString InValue)
 	Implementation->SendConsoleVariableChange(MoveTemp(InName), MoveTemp(InValue));
 }
 
+void FManager::SendListItemCheckStateChange(FString InName, ECheckBoxState InCheckedState)
+{
+	Implementation->SendListItemCheckStateChange(MoveTemp(InName), MoveTemp(InCheckedState));
+}
+
 void FManager::SetEnableMultiUserSupport(bool bIsEnabled)
 {
 	Implementation->bIsEnabled = bIsEnabled;
diff --git a/Engine/Plugins/Editor/ConsoleVariablesEditor/Source/ConsoleVariablesEditor/Private/MultiUser/ConsoleVariableSync.h b/Engine/Plugins/Editor/ConsoleVariablesEditor/Source/ConsoleVariablesEditor/Private/MultiUser/ConsoleVariableSync.h
--- a/Engine/Plugins/Editor/ConsoleVariablesEditor/Source/ConsoleVariablesEditor/Private/MultiUser/ConsoleVariableSync.h
+++ b/Engine/Plugins/Editor/ConsoleVariablesEditor/Source/ConsoleVariablesEditor/Private/MultiUser/ConsoleVariableSync.h
@@ -5,11 +5,13 @@
 #include "Containers/UnrealString.h"
 #include "Delegates/Delegate.h"
 #include "Delegates/DelegateCombinations.h"
+#include "Styling/SlateTypes.h"
 #include "Templates/UniquePtr.h"
 
 #include "ConcertMessages.h"
 
 DECLARE_MULTICAST_DELEGATE_TwoParams(FOnRemoteCVarChange, FString, FString);
+DECLARE_MULTICAST_DELEGATE_TwoParams(FOnRemoteListItemCheckStateChange, FString, ECheckBoxState);
 DECLARE_MULTICAST_DELEGATE_OneParam(FOnMultiUserConnectionChange, EConcertConnectionStatus);
 
 namespace UE::ConsoleVariables::MultiUser::Private
@@ -30,12 +32,18 @@ struct FManager
 
 	/** Delegate that is invoked when a remote client has sent a new console variable value.*/
 	FOnRemoteCVarChange& OnRemoteCVarChange();
+	
+	/** Delegate invoked when a remote client's console variable list item in CVE is manually checked or unchecked. */
+	FOnRemoteListItemCheckStateChange& OnRemoteListItemCheckStateChange();
 
 	/** Delegate that is invoked when the connection status changes for this client. */
 	FOnMultiUserConnectionChange& OnConnectionChange();
 
 	/** Sends the named console variable with value to all connected Multi-user clients */
 	void SendConsoleVariableChange(FString InName, FString InValue);
+	
+	/** Sends the named console variable with value to all connected Multi-user clients */
+	void SendListItemCheckStateChange(FString InName, ECheckBoxState InCheckedState);
 
 	/** Enables / disables multi-user message handling. */
 	void SetEnableMultiUserSupport(bool bIsEnabled);
diff --git a/Engine/Plugins/Editor/ConsoleVariablesEditor/Source/ConsoleVariablesEditor/Private/MultiUser/ConsoleVariableSyncData.h b/Engine/Plugins/Editor/ConsoleVariablesEditor/Source/ConsoleVariablesEditor/Private/MultiUser/ConsoleVariableSyncData.h
--- a/Engine/Plugins/Editor/ConsoleVariablesEditor/Source/ConsoleVariablesEditor/Private/MultiUser/ConsoleVariableSyncData.h
+++ b/Engine/Plugins/Editor/ConsoleVariablesEditor/Source/ConsoleVariablesEditor/Private/MultiUser/ConsoleVariableSyncData.h
@@ -6,6 +6,9 @@
 #include "Misc/Guid.h"
 
 #include "ConcertMessageData.h"
+
+#include "Styling/SlateTypes.h"
+
 #include "ConsoleVariableSyncData.generated.h"
 
 UCLASS(config=Engine, DisplayName="Multi-user Console Variable Synchronization")
@@ -97,3 +100,15 @@ struct FConcertSetConsoleVariableEvent
 	UPROPERTY()
 	FString Value;
 };
+
+USTRUCT()
+struct FConcertSetListItemCheckStateEvent
+{
+	GENERATED_BODY()
+	/** Console variable name to apply the checked state change. */
+	UPROPERTY()
+	FString Variable;
+	/** The checked state of the list item. */
+	UPROPERTY()
+	ECheckBoxState CheckState;
+};
diff --git a/Engine/Plugins/Editor/ConsoleVariablesEditor/Source/ConsoleVariablesEditor/Private/Views/List/ConsoleVariablesEditorListRow.cpp b/Engine/Plugins/Editor/ConsoleVariablesEditor/Source/ConsoleVariablesEditor/Private/Views/List/ConsoleVariablesEditorListRow.cpp
--- a/Engine/Plugins/Editor/ConsoleVariablesEditor/Source/ConsoleVariablesEditor/Private/Views/List/ConsoleVariablesEditorListRow.cpp
+++ b/Engine/Plugins/Editor/ConsoleVariablesEditor/Source/ConsoleVariablesEditor/Private/Views/List/ConsoleVariablesEditorListRow.cpp
@@ -125,7 +125,7 @@ void FConsoleVariablesEditorListRow::ResetToStartupValueAndSource() const
 	if (const TSharedPtr<FConsoleVariablesEditorCommandInfo> PinnedCommand = GetCommandInfo().Pin())
 	{
 		PinnedCommand->SetSourceFlag(PinnedCommand->StartupSource);
-		PinnedCommand->ExecuteCommand(PinnedCommand->StartupValueAsString, false);
+		PinnedCommand->ExecuteCommand(PinnedCommand->StartupValueAsString, true, false);
 	}
 }
 
@@ -325,7 +325,7 @@ FReply FConsoleVariablesEditorListRow::OnActionButtonClicked()
 	if (bIsGlobalSearch)
 	{
 		FConsoleVariablesEditorAssetSaveData MatchingData;
-		if (!ConsoleVariablesEditorModule.GetPresetAsset()->FindSavedDataByCommandString(CommandName, MatchingData))
+		if (!ConsoleVariablesEditorModule.GetPresetAsset()->FindSavedDataByCommandString(CommandName, MatchingData, ESearchCase::IgnoreCase))
 		{
 			EditableAsset->AddOrSetConsoleObjectSavedData(
 				{
diff --git a/Engine/Plugins/Editor/ConsoleVariablesEditor/Source/ConsoleVariablesEditor/Private/Views/List/SConsoleVariablesEditorList.cpp b/Engine/Plugins/Editor/ConsoleVariablesEditor/Source/ConsoleVariablesEditor/Private/Views/List/SConsoleVariablesEditorList.cpp
--- a/Engine/Plugins/Editor/ConsoleVariablesEditor/Source/ConsoleVariablesEditor/Private/Views/List/SConsoleVariablesEditorList.cpp
+++ b/Engine/Plugins/Editor/ConsoleVariablesEditor/Source/ConsoleVariablesEditor/Private/Views/List/SConsoleVariablesEditorList.cpp
@@ -9,6 +9,7 @@
 #include "ConsoleVariablesEditorListFilters/ConsoleVariablesEditorListFilter_ShowOnlyVariables.h"
 #include "ConsoleVariablesEditorListFilters/ConsoleVariablesEditorListFilter_Source.h"
 
+#include "ConsoleVariablesEditorLog.h"
 #include "ConsoleVariablesEditorModule.h"
 #include "ConsoleVariablesEditorProjectSettings.h"
 #include "ConsoleVariablesEditorStyle.h"
@@ -439,7 +440,7 @@ void SConsoleVariablesEditorList::RebuildListWithListMode(
 
 		for (const FConsoleVariablesEditorListRowPtr& Item : TreeViewRootObjects)
 		{
-			if (Item->GetCommandInfo().Pin()->Command.Equals(InConsoleCommandToScrollTo))
+			if (Item->GetCommandInfo().Pin()->Command.Equals(InConsoleCommandToScrollTo, ESearchCase::IgnoreCase))
 			{
 				ScrollToItem = Item;
 				break;
@@ -1025,13 +1026,6 @@ void SConsoleVariablesEditorList::CacheCurrentListItemData()
 		if (const TSharedPtr<FConsoleVariablesEditorCommandInfo>& CommandInfo = Item->GetCommandInfo().Pin())
 		{
 			const FString& CommandName = CommandInfo->Command;
-				
-			/*CachedCommandStates.SetNum(Algo::RemoveIf(
-				CachedCommandStates,
-				[CommandName](const FConsoleVariablesEditorAssetSaveData& CachedData)
-				{
-					return CachedData.CommandName.Equals(CommandName);
-				}));*/
 			
 			CachedCommandStates.Add(
 				{
@@ -1060,7 +1054,7 @@ void SConsoleVariablesEditorList::RestorePreviousListItemData()
 					CachedCommandStates,
 					[CommandName](const FConsoleVariablesEditorAssetSaveData& CachedData)
 					{
-						return CachedData.CommandName.Equals(CommandName);
+						return CachedData.CommandName.Equals(CommandName, ESearchCase::IgnoreCase);
 					}))
 				{
 					Item->SetCachedValue((*Match).CommandValueAsString);
@@ -1103,8 +1097,13 @@ void SConsoleVariablesEditorList::GenerateTreeView(const bool bSkipExecutionOfCa
 	
 	check(EditableAsset);
 
-	for (const FConsoleVariablesEditorAssetSaveData& SavedCommand : EditableAsset->GetSavedCommands())
+	TArray<FConsoleVariablesEditorAssetSaveData> SavedCommands = EditableAsset->GetSavedCommands();
+	
+	for (const FConsoleVariablesEditorAssetSaveData& SavedCommand : SavedCommands)
 	{
+		UE_LOG(LogConsoleVariablesEditor, VeryVerbose, TEXT("%hs: Considering %s for the list"),
+			__FUNCTION__, *SavedCommand.CommandName);
+		
 		// Get corresponding CommandInfo for tracking or make one if the command is non-value
 		TSharedPtr<FConsoleVariablesEditorCommandInfo> CommandInfo =
 			ConsoleVariablesEditorModule.FindCommandInfoByName(SavedCommand.CommandName).Pin();
@@ -1122,15 +1121,31 @@ void SConsoleVariablesEditorList::GenerateTreeView(const bool bSkipExecutionOfCa
 		{
 			FConsoleVariablesEditorListRowPtr RowToAdd = nullptr;
 			
-			if (TSharedPtr<FConsoleVariablesEditorListRow>* MatchingRow = Algo::FindByPredicate(
+			// If the row already exists, let's not make another one.
+			if (Algo::FindByPredicate(
+				TreeViewRootObjects,
+				[this, &CommandInfo](const FConsoleVariablesEditorListRowPtr& ExistingRow)
+				{
+					return ExistingRow->GetCommandInfo().IsValid() &&
+						ExistingRow->GetCommandInfo().Pin()->Command.Equals(CommandInfo->Command, ESearchCase::IgnoreCase);
+				}) != nullptr)
+			{
+				UE_LOG(LogConsoleVariablesEditor, VeryVerbose, TEXT("%hs: Not Adding %s to the list because it's already in the list"),
+					__FUNCTION__, *SavedCommand.CommandName);
+				continue;
+			}
+			// If the row exists in our cache, let's use that one instead of creating a new one.
+			if (TSharedPtr<FConsoleVariablesEditorListRow>* MatchingLastPresetRow = Algo::FindByPredicate(
 				LastPresetObjects,
 				[this, &CommandInfo](const FConsoleVariablesEditorListRowPtr& ExistingRow)
 				{
 					return ExistingRow->GetCommandInfo().IsValid() &&
-						ExistingRow->GetCommandInfo().Pin()->Command.Equals(CommandInfo->Command);
+						ExistingRow->GetCommandInfo().Pin()->Command.Equals(CommandInfo->Command, ESearchCase::IgnoreCase);
 				}))
 			{
-				RowToAdd = *MatchingRow;
+				UE_LOG(LogConsoleVariablesEditor, VeryVerbose, TEXT("%hs: Adding %s to the list from cache"),
+					__FUNCTION__, *SavedCommand.CommandName);
+				RowToAdd = *MatchingLastPresetRow;
 			}
 			else
 			{
@@ -1162,15 +1177,18 @@ void SConsoleVariablesEditorList::GenerateTreeView(const bool bSkipExecutionOfCa
 						!CachedCommandStates.ContainsByPredicate(
 							[this, &CommandInfo](const FConsoleVariablesEditorAssetSaveData& Comparator)
 							{
-								return Comparator.CommandName.Equals(CommandInfo->Command);
+								return Comparator.CommandName.Equals(CommandInfo->Command, ESearchCase::IgnoreCase);
 							});
 				}
-
+				
 				if (bIsEligibleForExecution)
 				{
 					CommandInfo->ExecuteCommand(SavedCommand.CommandValueAsString);
 				}
 				
+				UE_LOG(LogConsoleVariablesEditor, VeryVerbose, TEXT("%hs: Adding %s to the list with a new row item"),
+					__FUNCTION__, *SavedCommand.CommandName);
+				
 				TreeViewRootObjects.Add(RowToAdd);
 			}
 		}
diff --git a/Engine/Plugins/Editor/ConsoleVariablesEditor/Source/ConsoleVariablesEditor/Private/Views/List/SConsoleVariablesEditorListValueInput.cpp b/Engine/Plugins/Editor/ConsoleVariablesEditor/Source/ConsoleVariablesEditor/Private/Views/List/SConsoleVariablesEditorListValueInput.cpp
--- a/Engine/Plugins/Editor/ConsoleVariablesEditor/Source/ConsoleVariablesEditor/Private/Views/List/SConsoleVariablesEditorListValueInput.cpp
+++ b/Engine/Plugins/Editor/ConsoleVariablesEditor/Source/ConsoleVariablesEditor/Private/Views/List/SConsoleVariablesEditorListValueInput.cpp
@@ -109,9 +109,6 @@ void SConsoleVariablesEditorListValueInput_Float::Construct(const FArguments& In
 				!PinnedItem->GetCachedValue().Equals(ValueAsString))
 			{				
 				PinnedItem->GetCommandInfo().Pin()->ExecuteCommand(ValueAsString);
-
-				FConsoleVariablesEditorModule::Get().SendMultiUserConsoleVariableChange(
-					PinnedItem->GetCommandInfo().Pin()->Command, ValueAsString);
 					
 				PinnedItem->SetCachedValue(ValueAsString);
 			}
@@ -182,9 +179,6 @@ void SConsoleVariablesEditorListValueInput_Int::Construct(const FArguments& InAr
 			{				
 				PinnedItem->GetCommandInfo().Pin()->ExecuteCommand(ValueAsString);
 
-				FConsoleVariablesEditorModule::Get().SendMultiUserConsoleVariableChange(
-					PinnedItem->GetCommandInfo().Pin()->Command, ValueAsString);
-
 				PinnedItem->SetCachedValue(ValueAsString);
 			}
 		})
@@ -279,9 +273,6 @@ void SConsoleVariablesEditorListValueInput_String::Construct(const FArguments& I
 			{				
 				PinnedItem->GetCommandInfo().Pin()->ExecuteCommand(ValueAsString);
 
-				FConsoleVariablesEditorModule::Get().SendMultiUserConsoleVariableChange(
-					PinnedItem->GetCommandInfo().Pin()->Command, ValueAsString);
-
 				PinnedItem->SetCachedValue(ValueAsString);
 			}
 		})
@@ -372,8 +363,6 @@ void SConsoleVariablesEditorListValueInput_Bool::SetInputValue(const FString& In
 	{				
 		PinnedItem->GetCommandInfo().Pin()->ExecuteCommand(InValueAsString);
 
-		FConsoleVariablesEditorModule::Get().SendMultiUserConsoleVariableChange(PinnedItem->GetCommandInfo().Pin()->Command, InValueAsString);
-
 		PinnedItem->SetCachedValue(InValueAsString);
 	}
 }
@@ -440,9 +429,6 @@ void SConsoleVariablesEditorListValueInput_Command::Construct(const FArguments&
 					{
 						PinnedItem->GetCommandInfo().Pin()->ExecuteCommand(ValueAsString);
 
-						FConsoleVariablesEditorModule::Get().SendMultiUserConsoleVariableChange(
-							PinnedItem->GetCommandInfo().Pin()->Command, ValueAsString);
-
 						PinnedItem->SetCachedValue(ValueAsString);
 					}
 
diff --git a/Engine/Plugins/Editor/ConsoleVariablesEditor/Source/ConsoleVariablesEditor/Private/Views/MainPanel/ConsoleVariablesEditorMainPanel.cpp b/Engine/Plugins/Editor/ConsoleVariablesEditor/Source/ConsoleVariablesEditor/Private/Views/MainPanel/ConsoleVariablesEditorMainPanel.cpp
--- a/Engine/Plugins/Editor/ConsoleVariablesEditor/Source/ConsoleVariablesEditor/Private/Views/MainPanel/ConsoleVariablesEditorMainPanel.cpp
+++ b/Engine/Plugins/Editor/ConsoleVariablesEditor/Source/ConsoleVariablesEditor/Private/Views/MainPanel/ConsoleVariablesEditorMainPanel.cpp
@@ -10,10 +10,11 @@
 #include "Views/List/ConsoleVariablesEditorList.h"
 #include "Views/MainPanel/SConsoleVariablesEditorMainPanel.h"
 
+#include "Editor/EditorEngine.h"
 #include "FileHelpers.h"
 #include "Framework/Application/SlateApplication.h"
 #include "MultiUser/ConsoleVariableSyncData.h"
-
+#include "UnrealEngine.h"
 
 FConsoleVariablesEditorMainPanel::FConsoleVariablesEditorMainPanel()
 {
@@ -23,6 +24,8 @@ FConsoleVariablesEditorMainPanel::FConsoleVariablesEditorMainPanel()
 		&FConsoleVariablesEditorMainPanel::OnConnectionChanged);
 	OnRemoteCVarChangeHandle = MultiUserManager.OnRemoteCVarChange().AddStatic(
 		&FConsoleVariablesEditorMainPanel::OnRemoteCvarChange);
+	OnRemoteListItemCheckStateChangeHandle = MultiUserManager.OnRemoteListItemCheckStateChange().AddStatic(
+			&FConsoleVariablesEditorMainPanel::OnRemoteListItemCheckStateChange);
 }
 
 FConsoleVariablesEditorMainPanel::~FConsoleVariablesEditorMainPanel()
@@ -31,6 +34,7 @@ FConsoleVariablesEditorMainPanel::~FConsoleVariablesEditorMainPanel()
 	EditorList.Reset();
 	MultiUserManager.OnConnectionChange().Remove(OnConnectionChangedHandle);
 	MultiUserManager.OnRemoteCVarChange().Remove(OnRemoteCVarChangeHandle);
+	MultiUserManager.OnRemoteListItemCheckStateChange().Remove(OnRemoteListItemCheckStateChangeHandle);
 }
 
 TSharedRef<SWidget> FConsoleVariablesEditorMainPanel::GetOrCreateWidget()
@@ -311,3 +315,16 @@ void FConsoleVariablesEditorMainPanel::OnRemoteCvarChange(const FString InName,
 
 	ConsoleVariablesEditorModule.OnRemoteCvarChanged(InName, InValue);
 }
+
+void FConsoleVariablesEditorMainPanel::OnRemoteListItemCheckStateChange(const FString InName, ECheckBoxState InCheckedState)
+{
+	const FString CheckStateAsString =
+		InCheckedState == ECheckBoxState::Checked ? "Checked" : InCheckedState == ECheckBoxState::Unchecked ? "Unchecked" : "Undetermined";
+	UE_LOG(LogConsoleVariablesEditor, Display,
+		TEXT("%hs: Remote check state change for list item with name '%s' set to '%s'"), __FUNCTION__, *InName, *CheckStateAsString);
+}
+
+void FConsoleVariablesEditorMainPanel::SendListItemCheckStateChange(const FString& InName, ECheckBoxState InCheckedState)
+{
+	GetMultiUserManager().SendListItemCheckStateChange(InName, InCheckedState);
+}
diff --git a/Engine/Plugins/Editor/ConsoleVariablesEditor/Source/ConsoleVariablesEditor/Private/Views/MainPanel/ConsoleVariablesEditorMainPanel.h b/Engine/Plugins/Editor/ConsoleVariablesEditor/Source/ConsoleVariablesEditor/Private/Views/MainPanel/ConsoleVariablesEditorMainPanel.h
--- a/Engine/Plugins/Editor/ConsoleVariablesEditor/Source/ConsoleVariablesEditor/Private/Views/MainPanel/ConsoleVariablesEditorMainPanel.h
+++ b/Engine/Plugins/Editor/ConsoleVariablesEditor/Source/ConsoleVariablesEditor/Private/Views/MainPanel/ConsoleVariablesEditorMainPanel.h
@@ -86,8 +86,12 @@ private:
 
 	static void OnConnectionChanged(EConcertConnectionStatus Status);
 	static void OnRemoteCvarChange(const FString InName, const FString InValue);
-
+	static void OnRemoteListItemCheckStateChange(const FString InName, ECheckBoxState InCheckedState);
+	
+	void SendListItemCheckStateChange(const FString& InName, ECheckBoxState InCheckedState);
+	
 	UE::ConsoleVariables::MultiUser::Private::FManager MultiUserManager;
 	FDelegateHandle OnConnectionChangedHandle;
 	FDelegateHandle OnRemoteCVarChangeHandle;
+	FDelegateHandle OnRemoteListItemCheckStateChangeHandle;
 };
diff --git a/Engine/Plugins/Editor/ConsoleVariablesEditor/Source/ConsoleVariablesEditor/Public/ConsoleVariablesAsset.h b/Engine/Plugins/Editor/ConsoleVariablesEditor/Source/ConsoleVariablesEditor/Public/ConsoleVariablesAsset.h
--- a/Engine/Plugins/Editor/ConsoleVariablesEditor/Source/ConsoleVariablesEditor/Public/ConsoleVariablesAsset.h
+++ b/Engine/Plugins/Editor/ConsoleVariablesEditor/Source/ConsoleVariablesEditor/Public/ConsoleVariablesAsset.h
@@ -61,7 +61,8 @@ public:
 	}
 
 	/** Outputs the FConsoleVariablesEditorCommandInfo matching InCommand. Returns whether a match was found. Case sensitive. */
-	bool FindSavedDataByCommandString(const FString InCommandString, FConsoleVariablesEditorAssetSaveData& OutValue) const;
+	bool FindSavedDataByCommandString(
+		const FString InCommandString, FConsoleVariablesEditorAssetSaveData& OutValue, const ESearchCase::Type SearchCase = ESearchCase::IgnoreCase) const;
 
 	/** Set the value of a saved console variable if the name matches; add a new console variable to the list if a match is not found. */
 	void AddOrSetConsoleObjectSavedData(const FConsoleVariablesEditorAssetSaveData& InData);
