Console Variables Editor:

-Fix for: nDisplay nodes/viewports shift position when rendering with MRQ

diff --git a/Engine/Plugins/Editor/ConsoleVariablesEditor/Source/ConsoleVariablesEditor/Private/ConsoleVariablesEditorModule.cpp b/Engine/Plugins/Editor/ConsoleVariablesEditor/Source/ConsoleVariablesEditor/Private/ConsoleVariablesEditorModule.cpp
--- a/Engine/Plugins/Editor/ConsoleVariablesEditor/Source/ConsoleVariablesEditor/Private/ConsoleVariablesEditorModule.cpp
+++ b/Engine/Plugins/Editor/ConsoleVariablesEditor/Source/ConsoleVariablesEditor/Private/ConsoleVariablesEditorModule.cpp
@@ -260,11 +260,23 @@ bool FConsoleVariablesEditorModule::PopulateGlobalSearchAssetWithVariablesMatchi
 
 void FConsoleVariablesEditorModule::SendMultiUserConsoleVariableChange(const FString& InVariableName, const FString& InValueAsString) const
 {
+	if (GEditor && GEditor->IsPlaySessionInProgress())
+	{
+		UE_LOG(LogConsoleVariablesEditor, Warning, TEXT("%hs: Play In Editor is about to start or has started; Multi-User Cvar sync is suspended duirng PIE."), __FUNCTION__);
+		return;
+	}
+	
 	MainPanel->GetMultiUserManager().SendConsoleVariableChange(InVariableName, InValueAsString);
 }
 
 void FConsoleVariablesEditorModule::OnRemoteCvarChanged(const FString InName, const FString InValue)
 {
+	if (GEditor && GEditor->IsPlaySessionInProgress())
+	{
+		UE_LOG(LogConsoleVariablesEditor, Warning, TEXT("%hs: Play In Editor is about to start or has started; Multi-User Cvar sync is suspended duirng PIE."), __FUNCTION__);
+		return;
+	}
+	
 	UE_LOG(LogConsoleVariablesEditor, VeryVerbose, TEXT("Remote set console variable %s = %s"), *InName, *InValue);
 
 	CommandsRecentlyReceivedFromMultiUser.Add(InName, InValue);
