//UE5/Main - NonUnity Compile LyraEditor Win64 - SConsoleVariablesEditorMainPanel.cpp error: 'FPropertyEditorModule': undeclared identifier

diff --git a/Engine/Plugins/Editor/ConsoleVariablesEditor/Source/ConsoleVariablesEditor/Private/Views/MainPanel/SConsoleVariablesEditorMainPanel.cpp b/Engine/Plugins/Editor/ConsoleVariablesEditor/Source/ConsoleVariablesEditor/Private/Views/MainPanel/SConsoleVariablesEditorMainPanel.cpp
--- a/Engine/Plugins/Editor/ConsoleVariablesEditor/Source/ConsoleVariablesEditor/Private/Views/MainPanel/SConsoleVariablesEditorMainPanel.cpp
+++ b/Engine/Plugins/Editor/ConsoleVariablesEditor/Source/ConsoleVariablesEditor/Private/Views/MainPanel/SConsoleVariablesEditorMainPanel.cpp
@@ -12,9 +12,13 @@
 #include "Views/MainPanel/SConsoleVariablesEditorCustomConsoleInputBox.h"
 
 #include "ContentBrowserModule.h"
+#include "DetailsViewArgs.h"
 #include "Framework/Application/SlateApplication.h"
 #include "Framework/MultiBox/MultiBoxBuilder.h"
 #include "IContentBrowserSingleton.h"
+#include "IDetailsView.h"
+#include "Modules/ModuleManager.h"
+#include "PropertyEditorModule.h"
 #include "SPositiveActionButton.h"
 #include "Styling/AppStyle.h"
 #include "Styling/StyleColors.h"
