//UE5/Main - NonUnity Compile LyraEditor Win64 - ConsoleVariablesEditorFunctionLibrary.cpp error: cannot convert argument 1 from 'const UConsoleVariablesAsset *' to 'const FAssetData &'

diff --git a/Engine/Plugins/Editor/ConsoleVariablesEditor/Source/ConsoleVariablesEditor/Private/ConsoleVariablesEditorFunctionLibrary.cpp b/Engine/Plugins/Editor/ConsoleVariablesEditor/Source/ConsoleVariablesEditor/Private/ConsoleVariablesEditorFunctionLibrary.cpp
--- a/Engine/Plugins/Editor/ConsoleVariablesEditor/Source/ConsoleVariablesEditor/Private/ConsoleVariablesEditorFunctionLibrary.cpp
+++ b/Engine/Plugins/Editor/ConsoleVariablesEditor/Source/ConsoleVariablesEditor/Private/ConsoleVariablesEditorFunctionLibrary.cpp
@@ -3,9 +3,13 @@
 #include "ConsoleVariablesEditorFunctionLibrary.h"
 
 #include "ConsoleVariablesAsset.h"
+#include "ConsoleVariablesEditorCommandInfo.h"
 #include "ConsoleVariablesEditorLog.h"
 #include "ConsoleVariablesEditorModule.h"
+
 #include "MultiUser/ConsoleVariableSyncData.h"
+
+#include "AssetRegistry/AssetData.h"
 #include "Widgets/Input/SCheckBox.h"
 
 UConsoleVariablesAsset* UConsoleVariablesEditorFunctionLibrary::GetCurrentlyLoadedPreset()
@@ -27,7 +31,7 @@ void UConsoleVariablesEditorFunctionLibrary::LoadPresetIntoConsoleVariablesEdito
 	
 	const FConsoleVariablesEditorModule& ConsoleVariablesEditorModule = FConsoleVariablesEditorModule::Get();
 
-	ConsoleVariablesEditorModule.OpenConsoleVariablesDialogWithAssetSelected(InAsset);
+	ConsoleVariablesEditorModule.OpenConsoleVariablesDialogWithAssetSelected(FAssetData(InAsset));
 }
 
 bool UConsoleVariablesEditorFunctionLibrary::CopyCurrentListToAsset(UConsoleVariablesAsset* InAsset)
diff --git a/Engine/Plugins/Editor/ConsoleVariablesEditor/Source/ConsoleVariablesEditor/Public/ConsoleVariablesEditorFunctionLibrary.h b/Engine/Plugins/Editor/ConsoleVariablesEditor/Source/ConsoleVariablesEditor/Public/ConsoleVariablesEditorFunctionLibrary.h
--- a/Engine/Plugins/Editor/ConsoleVariablesEditor/Source/ConsoleVariablesEditor/Public/ConsoleVariablesEditorFunctionLibrary.h
+++ b/Engine/Plugins/Editor/ConsoleVariablesEditor/Source/ConsoleVariablesEditor/Public/ConsoleVariablesEditorFunctionLibrary.h
@@ -2,11 +2,12 @@
 
 #pragma once
 
-#include "CoreMinimal.h"
+#include "Containers/Array.h"
 #include "Kismet/BlueprintFunctionLibrary.h"
 
 #include "ConsoleVariablesEditorFunctionLibrary.generated.h"
 
+class FString;
 class UConsoleVariablesAsset;
 
 /** An asset used to track collections of console variables that can be recalled and edited using the Console Variables Editor. */
