Console Variables Editor:

 -Cvars cannot be added by clicking in the dynamic list

diff --git a/Engine/Plugins/Editor/ConsoleVariablesEditor/Source/ConsoleVariablesEditor/Private/Views/MainPanel/SConsoleVariablesEditorCustomConsoleInputBox.cpp b/Engine/Plugins/Editor/ConsoleVariablesEditor/Source/ConsoleVariablesEditor/Private/Views/MainPanel/SConsoleVariablesEditorCustomConsoleInputBox.cpp
--- a/Engine/Plugins/Editor/ConsoleVariablesEditor/Source/ConsoleVariablesEditor/Private/Views/MainPanel/SConsoleVariablesEditorCustomConsoleInputBox.cpp
+++ b/Engine/Plugins/Editor/ConsoleVariablesEditor/Source/ConsoleVariablesEditor/Private/Views/MainPanel/SConsoleVariablesEditorCustomConsoleInputBox.cpp
@@ -2,6 +2,7 @@
 
 #include "Views/MainPanel/SConsoleVariablesEditorCustomConsoleInputBox.h"
 
+#include "ConsoleVariablesEditorCommandInfo.h"
 #include "ConsoleVariablesEditorLog.h"
 #include "ConsoleVariablesEditorModule.h"
 #include "Views/List/ConsoleVariablesEditorList.h"
@@ -33,7 +34,7 @@ void SConsoleVariablesEditorCustomConsoleInputBox::Construct(
 		.Placement( MenuPlacement_BelowAnchor )
 		[
 			SAssignNew(InputText, SEditableTextBox)
-			.Font(FEditorStyle::Get().GetWidgetStyle<FTextBlockStyle>("Log.Normal").Font)
+			.Font(FAppStyle::Get().GetWidgetStyle<FTextBlockStyle>("Log.Normal").Font)
 			.HintText(LOCTEXT("ConsoleCommandExecutorHintText", "Enter Console Command"))
 			.OnTextChanged(this, &SConsoleVariablesEditorCustomConsoleInputBox::OnTextChanged)
 			.OnKeyCharHandler(this, &SConsoleVariablesEditorCustomConsoleInputBox::OnKeyCharHandler)
@@ -43,8 +44,12 @@ void SConsoleVariablesEditorCustomConsoleInputBox::Construct(
 				// Hide widget when focus is lost
 				if (CommitType == ETextCommit::OnUserMovedFocus)
 				{
-					SuggestionBox->SetIsOpen(false);
-					SetVisibility(EVisibility::Collapsed);
+					// If the newly focused widget is the suggestion list, we want the suggestion list to handle what's next
+					if (FSlateApplication::Get().GetUserFocusedWidget(0) != SuggestionListView)
+					{
+						SuggestionBox->SetIsOpen(false);
+						SetVisibility(EVisibility::Collapsed);
+					}
 				}
 			})
 			.ClearKeyboardFocusOnCommit(false)
@@ -53,7 +58,7 @@ void SConsoleVariablesEditorCustomConsoleInputBox::Construct(
 		.MenuContent
 		(
 			SNew(SBorder)
-			.BorderImage(FEditorStyle::GetBrush("Menu.Background"))
+			.BorderImage(FAppStyle::Get().GetBrush("Menu.Background"))
 			.Padding( FMargin(2) )
 			[
 				SNew(SBox)
@@ -100,7 +105,7 @@ void SConsoleVariablesEditorCustomConsoleInputBox::Construct(
 							[
 								SNew(STextBlock)
 								.Text(DisplayText)
-								.TextStyle(FEditorStyle::Get(), "Log.Normal")
+								.TextStyle(FAppStyle::Get(), "Log.Normal")
 								.HighlightText(Suggestions.SuggestionsHighlight)
 								.ColorAndOpacity(FSlateColor::UseForeground())
 								.ToolTip(ToolTip.IsValid() ? ToolTip : nullptr)
@@ -120,17 +125,9 @@ void SConsoleVariablesEditorCustomConsoleInputBox::Construct(
 
 						MarkActiveSuggestion();
 
-						// If the user selected this suggestion by clicking on it, then go ahead and close the suggestion
-						// box as they've chosen the suggestion they're interested in.
-						if( SelectInfo == ESelectInfo::OnMouseClick )
-						{
-							SuggestionBox->SetIsOpen( false );
-
-							// Jump the caret to the end of the newly auto-completed line. This makes it so that selecting
-							// an option doesn't leave the cursor in the middle of the suggestion (which makes it hard to 
-							// ctrl-back out, or to type "?" for help, etc.)
-							InputText->GoTo(ETextLocation::EndOfDocument);
-						}
+						// Just close the suggestion box and add the marked suggestion to the list and close up
+						CommitInput();
+						SetVisibility(EVisibility::Collapsed);
 					})
 					.ItemHeight(18)
 				]
@@ -300,12 +297,7 @@ FReply SConsoleVariablesEditorCustomConsoleInputBox::OnKeyDownHandler(const FGeo
 {
 	if (KeyPressed.GetKey().GetFName() == TEXT("Enter"))
 	{
-		const FText CommittedText = InputText->GetText();
-		bIgnoreUIUpdate = true;
-		InputText->SetText(FText::GetEmpty());
-		SuggestionBox->SetIsOpen( false );
-		bIgnoreUIUpdate = false;
-		return MainPanelWidget.Pin()->ValidateConsoleInputAndAddToCurrentPreset(CommittedText);
+		CommitInput();
 	}
 
 	return FReply::Unhandled();
@@ -377,4 +369,14 @@ void SConsoleVariablesEditorCustomConsoleInputBox::ClearSuggestions()
 	Suggestions.Reset();
 }
 
+void SConsoleVariablesEditorCustomConsoleInputBox::CommitInput()
+{
+	const FText CommittedText = InputText->GetText();
+	bIgnoreUIUpdate = true;
+	InputText->SetText(FText::GetEmpty());
+	SuggestionBox->SetIsOpen( false );
+	bIgnoreUIUpdate = false;
+	MainPanelWidget.Pin()->ValidateConsoleInputAndAddToCurrentPreset(CommittedText);
+}
+
 #undef LOCTEXT_NAMESPACE
diff --git a/Engine/Plugins/Editor/ConsoleVariablesEditor/Source/ConsoleVariablesEditor/Private/Views/MainPanel/SConsoleVariablesEditorCustomConsoleInputBox.h b/Engine/Plugins/Editor/ConsoleVariablesEditor/Source/ConsoleVariablesEditor/Private/Views/MainPanel/SConsoleVariablesEditorCustomConsoleInputBox.h
--- a/Engine/Plugins/Editor/ConsoleVariablesEditor/Source/ConsoleVariablesEditor/Private/Views/MainPanel/SConsoleVariablesEditorCustomConsoleInputBox.h
+++ b/Engine/Plugins/Editor/ConsoleVariablesEditor/Source/ConsoleVariablesEditor/Private/Views/MainPanel/SConsoleVariablesEditorCustomConsoleInputBox.h
@@ -91,6 +91,8 @@ public:
 	void MarkActiveSuggestion();
 
 	void ClearSuggestions();
+	
+	void CommitInput();
 
 private:
 
