Console Variables Editor:

-Fix Multi-User Icon

diff --git a/Engine/Plugins/Editor/ConsoleVariablesEditor/ConsoleVariables.uplugin b/Engine/Plugins/Editor/ConsoleVariablesEditor/ConsoleVariables.uplugin
--- a/Engine/Plugins/Editor/ConsoleVariablesEditor/ConsoleVariables.uplugin
+++ b/Engine/Plugins/Editor/ConsoleVariablesEditor/ConsoleVariables.uplugin
@@ -29,6 +29,10 @@
 			"Name": "ConcertMain",
 			"Enabled": true
 		},
+		{
+			"Name": "ConcertSharedSlate",
+			"Enabled": true
+		}
 	],
 	"Modules": [
 		{
diff --git a/Engine/Plugins/Editor/ConsoleVariablesEditor/Source/ConsoleVariablesEditor/ConsoleVariablesEditor.build.cs b/Engine/Plugins/Editor/ConsoleVariablesEditor/Source/ConsoleVariablesEditor/ConsoleVariablesEditor.build.cs
--- a/Engine/Plugins/Editor/ConsoleVariablesEditor/Source/ConsoleVariablesEditor/ConsoleVariablesEditor.build.cs
+++ b/Engine/Plugins/Editor/ConsoleVariablesEditor/Source/ConsoleVariablesEditor/ConsoleVariablesEditor.build.cs
@@ -22,6 +22,7 @@ public class ConsoleVariablesEditor : ModuleRules
 			new string[]
 			{
 				"Concert",
+				"ConcertSharedSlate",
 				"ConcertSyncClient",
 				"ConcertSyncCore",
 				"ConcertTransport",
diff --git a/Engine/Plugins/Editor/ConsoleVariablesEditor/Source/ConsoleVariablesEditor/Private/ConsoleVariablesEditorStyle.cpp b/Engine/Plugins/Editor/ConsoleVariablesEditor/Source/ConsoleVariablesEditor/Private/ConsoleVariablesEditorStyle.cpp
--- a/Engine/Plugins/Editor/ConsoleVariablesEditor/Source/ConsoleVariablesEditor/Private/ConsoleVariablesEditorStyle.cpp
+++ b/Engine/Plugins/Editor/ConsoleVariablesEditor/Source/ConsoleVariablesEditor/Private/ConsoleVariablesEditorStyle.cpp
@@ -44,6 +44,12 @@ const FName& FConsoleVariablesEditorStyle::GetStyleSetName() const
 	return ConsoleVariablesStyleSetName;
 }
 
+FString GetConcertContent(const FString& RelativePath, const ANSICHAR* Extension)
+{
+	static FString ContentDir = IPluginManager::Get().FindPlugin(TEXT("ConcertSyncClient"))->GetContentDir();
+	return (ContentDir / RelativePath) + Extension;
+}
+
 const FSlateBrush* FConsoleVariablesEditorStyle::GetBrush(const FName PropertyName, const ANSICHAR* Specifier, const ISlateStyle* RequestingStyle) const
 {
 	return StyleInstance->GetBrush(PropertyName, Specifier);
@@ -53,7 +59,7 @@ const FSlateBrush* FConsoleVariablesEditorStyle::GetBrush(const FName PropertyNa
 #define BOX_BRUSH( RelativePath, ... ) FSlateBoxBrush( Style->RootToContentDir( RelativePath, TEXT(".png") ), __VA_ARGS__ )
 #define BORDER_BRUSH( RelativePath, ... ) FSlateBorderBrush( Style->RootToContentDir( RelativePath, TEXT(".png") ), __VA_ARGS__ )
 #define IMAGE_BRUSH_SVG( RelativePath, ... ) FSlateVectorImageBrush( Style->RootToContentDir(RelativePath, TEXT(".svg") ), __VA_ARGS__)
-#define IMAGE_PLUGIN_BRUSH_SVG( RelativePath, ... ) FSlateVectorImageBrush( FConsoleVariablesEditorStyle::InContent(RelativePath, ".svg"), __VA_ARGS__)
+#define IMAGE_PLUGIN_BRUSH_SVG( PluginName, RelativePath, ... ) FSlateVectorImageBrush( FConsoleVariablesEditorStyle::GetExternalPluginContent(PluginName, RelativePath, ".svg"), __VA_ARGS__)
 
 const FVector2D Icon64x64(64.f, 64.f);
 const FVector2D Icon40x40(40.0f, 40.0f);
@@ -62,10 +68,11 @@ const FVector2D Icon16x16(16.0f, 16.0f);
 const FVector2D Icon12x12(12.0f, 12.0f);
 const FVector2D Icon8x8(8.f, 8.f);
 
-FString FConsoleVariablesEditorStyle::InContent(const FString& RelativePath, const ANSICHAR* Extension)
+FString FConsoleVariablesEditorStyle::GetExternalPluginContent(const FString& PluginName, const FString& RelativePath, const ANSICHAR* Extension)
 {
-	static FString ContentDir = IPluginManager::Get().FindPlugin(TEXT("ConcertSyncClient"))->GetContentDir();
-	return (ContentDir / RelativePath) + Extension;
+	FString ContentDir = IPluginManager::Get().FindPlugin(PluginName)->GetBaseDir() / RelativePath + Extension;
+	UE_LOG(LogTemp, Error, TEXT("%hs: Content = %s"), __FUNCTION__, *ContentDir);
+	return ContentDir;
 }
 
 TSharedRef< FSlateStyleSet > FConsoleVariablesEditorStyle::Create()
@@ -100,8 +107,9 @@ TSharedRef< FSlateStyleSet > FConsoleVariablesEditorStyle::Create()
 	Style->Set("ConsoleVariablesEditor.CommandGroupBorder", new BOX_BRUSH("Common/DarkGroupBorder", FMargin(4.0f / 16.0f)));
 	Style->Set("ConsoleVariablesEditor.DefaultBorder", new FSlateColorBrush(FStyleColors::Transparent));
 
+	// External plugin icons
 	// Multi-user Tab/Menu icons
-	Style->Set("Concert.MultiUser", new IMAGE_PLUGIN_BRUSH_SVG("Icons/icon_MultiUser", Icon16x16));
+	Style->Set("Icons.MultiUser", new IMAGE_PLUGIN_BRUSH_SVG("ConcertSharedSlate","Content/Icons/icon_MultiUser", Icon16x16));
 	
 	return Style;
 }
diff --git a/Engine/Plugins/Editor/ConsoleVariablesEditor/Source/ConsoleVariablesEditor/Private/ConsoleVariablesEditorStyle.h b/Engine/Plugins/Editor/ConsoleVariablesEditor/Source/ConsoleVariablesEditor/Private/ConsoleVariablesEditorStyle.h
--- a/Engine/Plugins/Editor/ConsoleVariablesEditor/Source/ConsoleVariablesEditor/Private/ConsoleVariablesEditorStyle.h
+++ b/Engine/Plugins/Editor/ConsoleVariablesEditor/Source/ConsoleVariablesEditor/Private/ConsoleVariablesEditorStyle.h
@@ -27,8 +27,8 @@ public:
 	}
 
 private:
-	
-	static FString InContent(const FString& RelativePath, const ANSICHAR* Extension);
+
+	static FString GetExternalPluginContent(const FString& PluginName, const FString& RelativePath, const ANSICHAR* Extension);
 	
 	static TSharedRef<FSlateStyleSet> Create();
 
diff --git a/Engine/Plugins/Editor/ConsoleVariablesEditor/Source/ConsoleVariablesEditor/Private/Views/MainPanel/SConsoleVariablesEditorMainPanel.cpp b/Engine/Plugins/Editor/ConsoleVariablesEditor/Source/ConsoleVariablesEditor/Private/Views/MainPanel/SConsoleVariablesEditorMainPanel.cpp
--- a/Engine/Plugins/Editor/ConsoleVariablesEditor/Source/ConsoleVariablesEditor/Private/Views/MainPanel/SConsoleVariablesEditorMainPanel.cpp
+++ b/Engine/Plugins/Editor/ConsoleVariablesEditor/Source/ConsoleVariablesEditor/Private/Views/MainPanel/SConsoleVariablesEditorMainPanel.cpp
@@ -211,7 +211,7 @@ void SConsoleVariablesEditorMainPanel::CreateConcertButtonIfNeeded()
 			.OnCheckStateChanged_Raw(this, &SConsoleVariablesEditorMainPanel::ToggleMultiUserDetails)
 			[
 				SNew(SImage)
-				.Image(FConsoleVariablesEditorStyle::Get().GetBrush("Concert.MultiUser"))
+				.Image(FConsoleVariablesEditorStyle::Get().GetBrush("Icons.MultiUser"))
 				.ColorAndOpacity(FSlateColor::UseForeground())
 			]
 		];
