Console Variables Editor:

-Add custom handler for pressing spacebar when adding a cvar

diff --git a/Engine/Plugins/Editor/ConsoleVariablesEditor/Source/ConsoleVariablesEditor/Private/Views/MainPanel/SConsoleVariablesEditorCustomConsoleInputBox.cpp b/Engine/Plugins/Editor/ConsoleVariablesEditor/Source/ConsoleVariablesEditor/Private/Views/MainPanel/SConsoleVariablesEditorCustomConsoleInputBox.cpp
--- a/Engine/Plugins/Editor/ConsoleVariablesEditor/Source/ConsoleVariablesEditor/Private/Views/MainPanel/SConsoleVariablesEditorCustomConsoleInputBox.cpp
+++ b/Engine/Plugins/Editor/ConsoleVariablesEditor/Source/ConsoleVariablesEditor/Private/Views/MainPanel/SConsoleVariablesEditorCustomConsoleInputBox.cpp
@@ -45,7 +45,8 @@ void SConsoleVariablesEditorCustomConsoleInputBox::Construct(
 				if (CommitType == ETextCommit::OnUserMovedFocus)
 				{
 					// If the newly focused widget is the suggestion list, we want the suggestion list to handle what's next
-					if (FSlateApplication::Get().GetUserFocusedWidget(0) != SuggestionListView)
+					if (FSlateApplication::Get().GetUserFocusedWidget(0) != SuggestionListView &&
+						FSlateApplication::Get().GetUserFocusedWidget(0) != InputText)
 					{
 						SuggestionBox->SetIsOpen(false);
 						SetVisibility(EVisibility::Collapsed);
@@ -126,8 +127,11 @@ void SConsoleVariablesEditorCustomConsoleInputBox::Construct(
 						MarkActiveSuggestion();
 
 						// Just close the suggestion box and add the marked suggestion to the list and close up
-						CommitInput();
-						SetVisibility(EVisibility::Collapsed);
+						if( SelectInfo == ESelectInfo::OnMouseClick )
+						{
+							CommitInput();
+							SetVisibility(EVisibility::Collapsed);
+						}
 					})
 					.ItemHeight(18)
 				]
