//UE5/Main - Cook QAVirtualProduction Win64 - Ensure condition failed: IPluginManager::Get().LoadModulesForEnabledPlugins(ELoadingPhase::PostEngineInit)

diff --git a/Engine/Plugins/Editor/ConsoleVariablesEditor/Source/ConsoleVariablesEditor/Private/MultiUser/ConsoleVariableSyncData.h b/Engine/Plugins/Editor/ConsoleVariablesEditor/Source/ConsoleVariablesEditor/Private/MultiUser/ConsoleVariableSyncData.h
--- a/Engine/Plugins/Editor/ConsoleVariablesEditor/Source/ConsoleVariablesEditor/Private/MultiUser/ConsoleVariableSyncData.h
+++ b/Engine/Plugins/Editor/ConsoleVariablesEditor/Source/ConsoleVariablesEditor/Private/MultiUser/ConsoleVariableSyncData.h
@@ -110,5 +110,5 @@ struct FConcertSetListItemCheckStateEvent
 	FString Variable;
 	/** The checked state of the list item. */
 	UPROPERTY()
-	ECheckBoxState CheckState;
+	ECheckBoxState CheckState = ECheckBoxState::Checked;
 };
