Console Variables Editor:

-Update all instances of FEditorStyle to FAppStyle

diff --git a/Engine/Plugins/Editor/ConsoleVariablesEditor/Source/ConsoleVariablesEditor/Private/MultiUser/ConcertConsoleVariableSessionCustomization.cpp b/Engine/Plugins/Editor/ConsoleVariablesEditor/Source/ConsoleVariablesEditor/Private/MultiUser/ConcertConsoleVariableSessionCustomization.cpp
--- a/Engine/Plugins/Editor/ConsoleVariablesEditor/Source/ConsoleVariablesEditor/Private/MultiUser/ConcertConsoleVariableSessionCustomization.cpp
+++ b/Engine/Plugins/Editor/ConsoleVariablesEditor/Source/ConsoleVariablesEditor/Private/MultiUser/ConcertConsoleVariableSessionCustomization.cpp
@@ -77,10 +77,9 @@ public:
 				.VAlign(VAlign_Center)
 				.Padding(FMargin(4.0f, 0.0f, 0.0f, 0.0f))
 				[
-					SNew(STextBlock)
-					.Font(this, &SConsoleVariableSessionRow::GetAvatarFont)
+					SNew(SImage)
+					.Image_Static(SConsoleVariableSessionRow::GetAvatarBrush)
 					.ColorAndOpacity(this, &SConsoleVariableSessionRow::GetAvatarColor)
-					.Text(FEditorFontGlyphs::Square)
 				]
 
 				// The client display name.
@@ -88,12 +87,12 @@ public:
 				.VAlign(VAlign_Center)
 				[
 					SNew(SBorder)
-					.BorderImage(FEditorStyle::GetBrush("NoBorder"))
+					.BorderImage(FAppStyle::Get().GetBrush("NoBorder"))
 					.ColorAndOpacity(FLinearColor(0.75f, 0.75f, 0.75f))
 					.Padding(FMargin(6.0f, 4.0f))
 					[
 						SNew(STextBlock)
-						.Font(FEditorStyle::GetFontStyle("BoldFont"))
+						.Font(FAppStyle::Get().GetFontStyle("BoldFont"))
 						.Text(GetDisplayName())
 					]
 				];
@@ -211,28 +210,9 @@ public:
 		return FText();
 	}
 
-	FSlateFontInfo GetAvatarFont() const
+	static const FSlateBrush* GetAvatarBrush()
 	{
-		static const FName ButtonIconSyle = TEXT("FontAwesome.10");
-		// This font is used to render a small square box filled with the avatar color.
-		FSlateFontInfo ClientIconFontInfo = FEditorStyle::Get().GetFontStyle(ButtonIconSyle);
-		ClientIconFontInfo.Size = 8;
-		ClientIconFontInfo.OutlineSettings.OutlineSize = 1;
-
-		TSharedPtr<FConcertCVarDetails> SettingPin = ClientSetting.Pin();
-		if (SettingPin.IsValid())
-		{
-			FConcertSessionClientInfo& Client = SettingPin->Details;
-			FLinearColor ClientOutlineColor = Client.ClientInfo.AvatarColor * 0.6f; // Make the font outline darker.
-			ClientOutlineColor.A = Client.ClientInfo.AvatarColor.A; // Put back the original alpha.
-			ClientIconFontInfo.OutlineSettings.OutlineColor = ClientOutlineColor;
-		}
-		else
-		{
-			ClientIconFontInfo.OutlineSettings.OutlineColor = FLinearColor(0.75, 0.75, 0.75); // This is an arbitrary color.
-		}
-
-		return ClientIconFontInfo;
+		return FAppStyle::Get().GetBrush("Icons.Toolbar.Stop");
 	}
 
 	FSlateColor GetAvatarColor() const
diff --git a/Engine/Plugins/Editor/ConsoleVariablesEditor/Source/ConsoleVariablesEditor/Private/Views/List/SConsoleVariablesEditorList.cpp b/Engine/Plugins/Editor/ConsoleVariablesEditor/Source/ConsoleVariablesEditor/Private/Views/List/SConsoleVariablesEditorList.cpp
--- a/Engine/Plugins/Editor/ConsoleVariablesEditor/Source/ConsoleVariablesEditor/Private/Views/List/SConsoleVariablesEditorList.cpp
+++ b/Engine/Plugins/Editor/ConsoleVariablesEditor/Source/ConsoleVariablesEditor/Private/Views/List/SConsoleVariablesEditorList.cpp
@@ -229,7 +229,7 @@ void SConsoleVariablesEditorList::Construct(const FArguments& InArgs, TSharedRef
 			.Padding(2.0f, 24.0f, 2.0f, 2.0f)
 			[
 				SNew(SRichTextBlock)
-				.DecoratorStyleSet(&FEditorStyle::Get())
+				.DecoratorStyleSet(&FAppStyle::Get())
 				.AutoWrapText(true)
 				.Justification(ETextJustify::Center)
 				.Text_Lambda([this]()
