//UE5/Main - Cook QAVirtualProduction Win64 - Ensure condition failed: IPluginManager::Get().LoadModulesForEnabledPlugins(ELoadingPhase::PostEngineInit)

diff --git a/Engine/Plugins/Editor/ConsoleVariablesEditor/Source/ConsoleVariablesEditor/Private/ConsoleVariablesEditorStyle.cpp b/Engine/Plugins/Editor/ConsoleVariablesEditor/Source/ConsoleVariablesEditor/Private/ConsoleVariablesEditorStyle.cpp
--- a/Engine/Plugins/Editor/ConsoleVariablesEditor/Source/ConsoleVariablesEditor/Private/ConsoleVariablesEditorStyle.cpp
+++ b/Engine/Plugins/Editor/ConsoleVariablesEditor/Source/ConsoleVariablesEditor/Private/ConsoleVariablesEditorStyle.cpp
@@ -71,7 +71,6 @@ const FVector2D Icon8x8(8.f, 8.f);
 FString FConsoleVariablesEditorStyle::GetExternalPluginContent(const FString& PluginName, const FString& RelativePath, const ANSICHAR* Extension)
 {
 	FString ContentDir = IPluginManager::Get().FindPlugin(PluginName)->GetBaseDir() / RelativePath + Extension;
-	UE_LOG(LogTemp, Error, TEXT("%hs: Content = %s"), __FUNCTION__, *ContentDir);
 	return ContentDir;
 }
 
