Console Variables Editor:

-Move ConsoleVariablesAsset from UncookedOnly module to separate Runtime module
-Make Console Variable Asset functions mostly BlueprintCallable or BlueprintPure
-Add option to execute all commands in an asset from context click in content browser

diff --git a/Engine/Plugins/Editor/ConsoleVariablesEditor/ConsoleVariables.uplugin b/Engine/Plugins/Editor/ConsoleVariablesEditor/ConsoleVariables.uplugin
--- a/Engine/Plugins/Editor/ConsoleVariablesEditor/ConsoleVariables.uplugin
+++ b/Engine/Plugins/Editor/ConsoleVariablesEditor/ConsoleVariables.uplugin
@@ -39,6 +39,11 @@
 			"Name": "ConsoleVariablesEditor",
 			"Type": "UncookedOnly",
 			"LoadingPhase": "Default"
+		},
+		{
+			"Name": "ConsoleVariablesEditorRuntime",
+			"Type": "Runtime",
+			"LoadingPhase": "Default"
 		}
 	]
 }
diff --git a/Engine/Plugins/Editor/ConsoleVariablesEditor/Source/ConsoleVariablesEditor/ConsoleVariablesEditor.build.cs b/Engine/Plugins/Editor/ConsoleVariablesEditor/Source/ConsoleVariablesEditor/ConsoleVariablesEditor.build.cs
--- a/Engine/Plugins/Editor/ConsoleVariablesEditor/Source/ConsoleVariablesEditor/ConsoleVariablesEditor.build.cs
+++ b/Engine/Plugins/Editor/ConsoleVariablesEditor/Source/ConsoleVariablesEditor/ConsoleVariablesEditor.build.cs
@@ -21,14 +21,15 @@ public class ConsoleVariablesEditor : ModuleRules
 		PrivateDependencyModuleNames.AddRange(
 			new string[]
 			{
+				"AssetRegistry",
+				"AssetTools",
+				"CoreUObject",
 				"Concert",
 				"ConcertSharedSlate",
 				"ConcertSyncClient",
 				"ConcertSyncCore",
 				"ConcertTransport",
-				"AssetRegistry",
-				"AssetTools",
-				"CoreUObject",
+				"ConsoleVariablesEditorRuntime",
 				"ContentBrowser",
 				"Engine",
 				"EditorStyle",
diff --git a/Engine/Plugins/Editor/ConsoleVariablesEditor/Source/ConsoleVariablesEditor/Private/AssetTypeActions/AssetTypeActions_ConsoleVariables.cpp b/Engine/Plugins/Editor/ConsoleVariablesEditor/Source/ConsoleVariablesEditor/Private/AssetTypeActions/AssetTypeActions_ConsoleVariables.cpp
--- a/Engine/Plugins/Editor/ConsoleVariablesEditor/Source/ConsoleVariablesEditor/Private/AssetTypeActions/AssetTypeActions_ConsoleVariables.cpp
+++ b/Engine/Plugins/Editor/ConsoleVariablesEditor/Source/ConsoleVariablesEditor/Private/AssetTypeActions/AssetTypeActions_ConsoleVariables.cpp
@@ -3,6 +3,7 @@
 #include "AssetTypeActions_ConsoleVariables.h"
 
 #include "ConsoleVariablesAsset.h"
+#include "ConsoleVariablesEditorStyle.h"
 #include "ConsoleVariablesEditor/Private/ConsoleVariablesEditorModule.h"
 
 #include "Framework/MultiBox/MultiBoxBuilder.h"
@@ -24,21 +25,81 @@ void FAssetTypeActions_ConsoleVariables::GetActions(const TArray<UObject*>& InOb
 	FAssetTypeActions_Base::GetActions(InObjects, MenuBuilder);
 
 	const TArray<TWeakObjectPtr<UConsoleVariablesAsset>> ConsoleVariableAssets = GetTypedWeakObjectPtrs<UConsoleVariablesAsset>(InObjects);
+	const int32 ConsoleVariablesAssetCount = ConsoleVariableAssets.Num();
 
 	MenuBuilder.AddMenuEntry(
 		LOCTEXT("AssetTypeActions_OpenVariableCollection", "Open Variable Collection in Editor"),
 		LOCTEXT("AssetTypeActions_OpenVariableCollectionToolTip", "Open this console variable collection in the Console Variables Editor. Select only one asset at a time."),
 		FSlateIcon(FAppStyle::Get().GetStyleSetName(), "SystemWideCommands.SummonOpenAssetDialog"),
 		FUIAction(
-			FExecuteAction::CreateLambda([=] {
+			FExecuteAction::CreateLambda([this, InObjects] {
 					OpenAssetEditor(InObjects, TSharedPtr<IToolkitHost>());
 			}),
-			FCanExecuteAction::CreateLambda([=] {
+			FCanExecuteAction::CreateLambda([ConsoleVariablesAssetCount] {
 					// We only want to open a single Variable Collection asset at a time, so let's ensure
 					// the number of selected assets is exactly one.
-					return ConsoleVariableAssets.Num() == 1;
-				})
+					return ConsoleVariablesAssetCount == 1;
+				}
 			)
+		)
+	);
+	
+	MenuBuilder.AddMenuEntry(
+		LOCTEXT("AssetTypeActions_ExecuteVariableCollection", "Execute Variable Collection"),
+		LOCTEXT("AssetTypeActions_ExecuteVariableCollectionToolTip", "Executes all commands and variables saved to the selected assets."),
+		FSlateIcon(FConsoleVariablesEditorStyle::Get().GetStyleSetName(), "ConsoleVariables.ToolbarButton.Small"),
+		FUIAction(
+			FExecuteAction::CreateLambda([ConsoleVariableAssets] {
+				UWorld* CurrentWorld = nullptr;
+				check (GIsEditor);
+				{
+					CurrentWorld = GEditor->GetEditorWorldContext().World();
+				}
+				
+				for (TWeakObjectPtr<UConsoleVariablesAsset> Asset : ConsoleVariableAssets)
+				{
+					if (Asset.IsValid())
+					{
+						Asset->ExecuteSavedCommands(CurrentWorld);
+					}
+				}
+			}),
+			FCanExecuteAction::CreateLambda([ConsoleVariablesAssetCount] {
+					// We only want to open a single Variable Collection asset at a time, so let's ensure
+					// the number of selected assets is exactly one.
+					return ConsoleVariablesAssetCount > 0;
+				}
+			)
+		)
+	);
+
+	MenuBuilder.AddMenuEntry(
+		LOCTEXT("AssetTypeActions_ExecuteOnlyCheckedInVariableCollection", "Execute Variable Collection (Only Checked)"),
+		LOCTEXT("AssetTypeActions_ExecuteOnlyCheckedInVariableCollectionToolTip", "Executes commands and variables with a Checked checkstate saved to the selected assets. Checkstates can be edited in the Console Variables Editor UI."),
+		FSlateIcon(FConsoleVariablesEditorStyle::Get().GetStyleSetName(), "ConsoleVariables.ToolbarButton.Small"),
+		FUIAction(
+			FExecuteAction::CreateLambda([ConsoleVariableAssets] {
+				UWorld* CurrentWorld = nullptr;
+				check (GIsEditor)
+				{
+					CurrentWorld = GEditor->GetEditorWorldContext().World();
+				}
+				
+				for (TWeakObjectPtr<UConsoleVariablesAsset> Asset : ConsoleVariableAssets)
+				{
+					if (Asset.IsValid())
+					{
+						Asset->ExecuteSavedCommands(CurrentWorld, true);
+					}
+				}
+			}),
+			FCanExecuteAction::CreateLambda([ConsoleVariablesAssetCount] {
+					// We only want to open a single Variable Collection asset at a time, so let's ensure
+					// the number of selected assets is exactly one.
+					return ConsoleVariablesAssetCount > 0;
+				}
+			)
+		)
 	);
 }
 
diff --git a/Engine/Plugins/Editor/ConsoleVariablesEditor/Source/ConsoleVariablesEditor/Private/ConsoleVariablesAsset.cpp b/Engine/Plugins/Editor/ConsoleVariablesEditor/Source/ConsoleVariablesEditor/Private/ConsoleVariablesAsset.cpp
deleted file mode 100644
diff --git a/Engine/Plugins/Editor/ConsoleVariablesEditor/Source/ConsoleVariablesEditor/Public/ConsoleVariablesAsset.h b/Engine/Plugins/Editor/ConsoleVariablesEditor/Source/ConsoleVariablesEditor/Public/ConsoleVariablesAsset.h
deleted file mode 100644
diff --git a/Engine/Plugins/Editor/ConsoleVariablesEditor/Source/ConsoleVariablesEditorRuntime/ConsoleVariablesEditorRuntime.build.cs b/Engine/Plugins/Editor/ConsoleVariablesEditor/Source/ConsoleVariablesEditorRuntime/ConsoleVariablesEditorRuntime.build.cs
new file mode 100644
--- /dev/null
+++ b/Engine/Plugins/Editor/ConsoleVariablesEditor/Source/ConsoleVariablesEditorRuntime/ConsoleVariablesEditorRuntime.build.cs
@@ -0,0 +1,30 @@
+// Copyright Epic Games, Inc. All Rights Reserved.
+
+using System.IO;
+using UnrealBuildTool;
+
+public class ConsoleVariablesEditorRuntime : ModuleRules
+{
+	public ConsoleVariablesEditorRuntime(ReadOnlyTargetRules Target) : base(Target)
+	{
+		PCHUsage = ModuleRules.PCHUsageMode.UseExplicitOrSharedPCHs;
+
+		PublicDependencyModuleNames.AddRange(
+			new string[]
+			{
+				"Core",
+				"OutputLog"
+			}
+		);
+			
+		
+		PrivateDependencyModuleNames.AddRange(
+			new string[]
+			{
+				"CoreUObject",
+				"Engine",
+				"SlateCore"
+			}
+		);
+	}
+}
diff --git a/Engine/Plugins/Editor/ConsoleVariablesEditor/Source/ConsoleVariablesEditorRuntime/Private/ConsoleVariablesAsset.cpp b/Engine/Plugins/Editor/ConsoleVariablesEditor/Source/ConsoleVariablesEditorRuntime/Private/ConsoleVariablesAsset.cpp
new file mode 100644
--- /dev/null
+++ b/Engine/Plugins/Editor/ConsoleVariablesEditor/Source/ConsoleVariablesEditorRuntime/Private/ConsoleVariablesAsset.cpp
@@ -0,0 +1,113 @@
+// Copyright Epic Games, Inc. All Rights Reserved.
+
+#include "ConsoleVariablesAsset.h"
+
+#include "ConsoleVariablesEditorRuntimeLog.h"
+#include "Engine/Engine.h"
+
+void UConsoleVariablesAsset::SetVariableCollectionDescription(const FString& InVariableCollectionDescription)
+{
+	VariableCollectionDescription = InVariableCollectionDescription;
+}
+
+TArray<FString> UConsoleVariablesAsset::GetSavedCommandsAsStringArray(bool bOnlyIncludeChecked) const
+{
+	TArray<FString> ReturnValue;
+	
+	// Make a copy because this array can change during iteration
+	TArray<FConsoleVariablesEditorAssetSaveData> SavedCommandsLocal = SavedCommands; 
+	for (int32 CommandIndex = 0; CommandIndex < SavedCommandsLocal.Num(); CommandIndex++)
+	{
+		const FConsoleVariablesEditorAssetSaveData& Command = SavedCommandsLocal[CommandIndex];
+		
+		if (bOnlyIncludeChecked && Command.CheckedState != ECheckBoxState::Checked)
+		{
+			continue;
+		}
+		
+		ReturnValue.Add(FString::Printf(TEXT("%s %s"), *Command.CommandName, *Command.CommandValueAsString));
+	}
+
+	return ReturnValue;
+}
+
+FString UConsoleVariablesAsset::GetSavedCommandsAsCommaSeparatedString(bool bOnlyIncludeChecked) const
+{
+	return FString::Join(GetSavedCommandsAsStringArray(bOnlyIncludeChecked), TEXT(","));
+}
+
+void UConsoleVariablesAsset::ExecuteSavedCommands(UObject* WorldContextObject, bool bOnlyIncludeChecked) const
+{
+	if (const TObjectPtr<UWorld> World = GEngine->GetWorldFromContextObject(WorldContextObject, EGetWorldErrorMode::LogAndReturnNull))
+	{
+		TArray<FString> Commands = GetSavedCommandsAsStringArray(bOnlyIncludeChecked);
+	
+		for (const FString& Command : Commands)
+		{
+			GEngine->Exec(World, *Command);
+		}
+	}
+	else
+	{
+		UE_LOG(LogConsoleVariablesEditorRuntime, Error, TEXT("%hs: Could not get world from WorldContextObject"), __FUNCTION__);
+	}
+}
+
+void UConsoleVariablesAsset::ReplaceSavedCommands(const TArray<FConsoleVariablesEditorAssetSaveData>& Replacement)
+{
+	SavedCommands = Replacement;
+}
+
+bool UConsoleVariablesAsset::FindSavedDataByCommandString(
+	const FString InCommandString, FConsoleVariablesEditorAssetSaveData& OutValue, const ESearchCase::Type SearchCase) const
+{
+	for (const FConsoleVariablesEditorAssetSaveData& Command : SavedCommands)
+	{
+		if (Command.CommandName.TrimStartAndEnd().Equals(InCommandString.TrimStartAndEnd(), SearchCase))
+		{
+			OutValue = Command;
+			return true;
+		}
+	}
+	
+	return false;
+}
+
+void UConsoleVariablesAsset::AddOrSetConsoleObjectSavedData(const FConsoleVariablesEditorAssetSaveData& InData)
+{
+	RemoveConsoleVariable(InData.CommandName);
+	SavedCommands.Add(InData);
+	
+	// Make a copy because this array can change during iteration
+	TArray<FConsoleVariablesEditorAssetSaveData> SavedCommandsLocal = SavedCommands; 
+	for (int32 CommandIndex = 0; CommandIndex < SavedCommandsLocal.Num(); CommandIndex++)
+	{
+		FConsoleVariablesEditorAssetSaveData& SavedCommand = SavedCommandsLocal[CommandIndex];
+		UE_LOG(LogConsoleVariablesEditorRuntime, VeryVerbose, TEXT("%hs: Command named '%s' at Index %i"),
+			__FUNCTION__, *SavedCommand.CommandName, CommandIndex);
+	}
+}
+
+bool UConsoleVariablesAsset::RemoveConsoleVariable(const FString InCommandString)
+{
+	FConsoleVariablesEditorAssetSaveData ExistingData;
+	int32 RemoveCount = 0;
+	while (FindSavedDataByCommandString(InCommandString, ExistingData, ESearchCase::IgnoreCase))
+	{
+		UE_LOG(LogConsoleVariablesEditorRuntime, VeryVerbose, TEXT("%hs: Removing %s from editable asset"),
+			__FUNCTION__, *InCommandString);
+		
+		if (SavedCommands.Remove(ExistingData) > 0)
+		{
+			RemoveCount++;
+		}
+	}
+	
+	return RemoveCount > 0;
+}
+
+void UConsoleVariablesAsset::CopyFrom(const UConsoleVariablesAsset* InAssetToCopy)
+{
+	VariableCollectionDescription = InAssetToCopy->GetVariableCollectionDescription();
+	SavedCommands = InAssetToCopy->GetSavedCommands();
+}
diff --git a/Engine/Plugins/Editor/ConsoleVariablesEditor/Source/ConsoleVariablesEditorRuntime/Private/ConsoleVariablesEditorRuntimeModule.cpp b/Engine/Plugins/Editor/ConsoleVariablesEditor/Source/ConsoleVariablesEditorRuntime/Private/ConsoleVariablesEditorRuntimeModule.cpp
new file mode 100644
--- /dev/null
+++ b/Engine/Plugins/Editor/ConsoleVariablesEditor/Source/ConsoleVariablesEditorRuntime/Private/ConsoleVariablesEditorRuntimeModule.cpp
@@ -0,0 +1,11 @@
+// Copyright Epic Games, Inc. All Rights Reserved.
+
+#pragma once
+
+#include "ConsoleVariablesEditorRuntimeModule.h"
+
+FConsoleVariablesEditorRuntimeModule::FConsoleVariablesEditorRuntimeModule()
+{
+}
+
+IMPLEMENT_MODULE(FConsoleVariablesEditorRuntimeModule, ConsoleVariablesEditorRuntime)
diff --git a/Engine/Plugins/Editor/ConsoleVariablesEditor/Source/ConsoleVariablesEditorRuntime/Private/ConsoleVariablesEditorRuntimeModule.h b/Engine/Plugins/Editor/ConsoleVariablesEditor/Source/ConsoleVariablesEditorRuntime/Private/ConsoleVariablesEditorRuntimeModule.h
new file mode 100644
--- /dev/null
+++ b/Engine/Plugins/Editor/ConsoleVariablesEditor/Source/ConsoleVariablesEditorRuntime/Private/ConsoleVariablesEditorRuntimeModule.h
@@ -0,0 +1,14 @@
+// Copyright Epic Games, Inc. All Rights Reserved.
+
+#pragma once
+
+#include "Modules/ModuleInterface.h"
+
+class FConsoleVariablesEditorRuntimeModule : public IModuleInterface
+{
+	
+public:
+	
+	FConsoleVariablesEditorRuntimeModule();
+	
+};
diff --git a/Engine/Plugins/Editor/ConsoleVariablesEditor/Source/ConsoleVariablesEditorRuntime/Public/ConsoleVariablesAsset.h b/Engine/Plugins/Editor/ConsoleVariablesEditor/Source/ConsoleVariablesEditorRuntime/Public/ConsoleVariablesAsset.h
new file mode 100644
--- /dev/null
+++ b/Engine/Plugins/Editor/ConsoleVariablesEditor/Source/ConsoleVariablesEditorRuntime/Public/ConsoleVariablesAsset.h
@@ -0,0 +1,111 @@
+// Copyright Epic Games, Inc. All Rights Reserved.
+
+#pragma once
+
+#include "Containers/UnrealString.h"
+#include "Styling/SlateTypes.h"
+
+#include "ConsoleVariablesAsset.generated.h"
+
+/** Data that will be serialized with this asset */
+USTRUCT(BlueprintType)
+struct FConsoleVariablesEditorAssetSaveData
+{
+	GENERATED_BODY()
+
+	FORCEINLINE bool operator==(const FConsoleVariablesEditorAssetSaveData& Comparator) const
+	{
+		return CommandName.Equals(Comparator.CommandName);
+	}
+	
+	UPROPERTY()
+	FString CommandName;
+
+	UPROPERTY()
+	FString CommandValueAsString;
+
+	UPROPERTY()
+	/** If Undetermined, we can assume this data was not previously saved */
+	ECheckBoxState CheckedState = ECheckBoxState::Undetermined;
+};
+
+/** An asset used to track collections of console variables that can be recalled and edited using the Console Variables Editor. */
+UCLASS(BlueprintType)
+class CONSOLEVARIABLESEDITORRUNTIME_API UConsoleVariablesAsset : public UObject
+{
+	GENERATED_BODY()
+public:
+
+	/** Sets a description for this variable collection. */
+	UFUNCTION(BlueprintCallable, Category = "Console Variables Editor")
+	void SetVariableCollectionDescription(const FString& InVariableCollectionDescription);
+	
+	UFUNCTION(BlueprintPure, Category = "Console Variables Editor")
+	FString GetVariableCollectionDescription() const
+	{
+		return VariableCollectionDescription;
+	}
+
+	/** Returns the saved list of console variable information such as the variable name, the type and the value of the variable at the time the asset was saved. */
+	UFUNCTION(BlueprintPure, Category="Console Variables Asset")
+	const TArray<FConsoleVariablesEditorAssetSaveData>& GetSavedCommands() const
+	{
+		return SavedCommands;
+	}
+
+	/** Returns the saved list of console variables as an array of FString.
+	 * @param bOnlyIncludeChecked If true, only commands and variables with a Checked checkstate in the Console Variables Editor UI will be included. Otherwise, all will be included. 
+	 */
+	UFUNCTION(BlueprintPure, Category="Console Variables Asset")
+	TArray<FString> GetSavedCommandsAsStringArray(bool bOnlyIncludeChecked = false) const;
+
+	/** Returns the saved list of console variables as a concatenated comma-separated string. Useful for passing commands and variables to a command line.
+	 * @param bOnlyIncludeChecked If true, only commands and variables with a Checked checkstate in the Console Variables Editor UI will be included. Otherwise, all will be included. 
+	 */
+	UFUNCTION(BlueprintPure, Category="Console Variables Asset")
+	FString GetSavedCommandsAsCommaSeparatedString(bool bOnlyIncludeChecked = false) const;
+
+	/** Executes all saved commands in this asset, optionally only including those with a Checked checkstate in the UI.
+	 * @param bOnlyIncludeChecked If true, only commands and variables with a Checked checkstate in the Console Variables Editor UI will be included. Otherwise, all will be included. 
+	 */
+	UFUNCTION(BlueprintCallable, Category="Console Variables Asset", meta=(WorldContext="WorldContextObject"))
+	void ExecuteSavedCommands(UObject* WorldContextObject, bool bOnlyIncludeChecked = false) const;
+
+	/** Completely replaces the saved data with new saved data */
+	UFUNCTION(BlueprintCallable, Category="Console Variables Asset")
+	void ReplaceSavedCommands(const TArray<FConsoleVariablesEditorAssetSaveData>& Replacement);
+
+	/** Returns how many console variables are serialized in this asset */
+	UFUNCTION(BlueprintPure, Category="Console Variables Asset")
+	int32 GetSavedCommandsCount() const
+	{
+		return GetSavedCommands().Num();
+	}
+
+	/** Outputs the FConsoleVariablesEditorAssetSaveData matching InCommand. Returns whether a match was found. Case sensitive. */
+	UFUNCTION(BlueprintCallable, Category="Console Variables Asset")
+	bool FindSavedDataByCommandString(
+		const FString InCommandString, FConsoleVariablesEditorAssetSaveData& OutValue, const ESearchCase::Type SearchCase = ESearchCase::IgnoreCase) const;
+
+	/** Set the value of a saved console variable if the name matches; add a new console variable to the list if a match is not found. */
+	UFUNCTION(BlueprintCallable, Category="Console Variables Asset")
+	void AddOrSetConsoleObjectSavedData(const FConsoleVariablesEditorAssetSaveData& InData);
+
+	/** Returns true if the element was found and successfully removed. */
+	UFUNCTION(BlueprintCallable, Category="Console Variables Asset")
+	bool RemoveConsoleVariable(const FString InCommandString);
+
+	/** Copy data from input asset to this asset */
+	UFUNCTION(BlueprintCallable, Category="Console Variables Asset")
+	void CopyFrom(const UConsoleVariablesAsset* InAssetToCopy);
+	
+private:
+	
+	/* User defined description of the variable collection */
+	UPROPERTY(AssetRegistrySearchable, BlueprintGetter = "GetVariableCollectionDescription", EditAnywhere, Category = "Console Variables Editor")
+	FString VariableCollectionDescription;
+
+	/** A saved list of console variable information such as the variable name, the type and the value of the variable at the time the asset was saved. */
+	UPROPERTY()
+	TArray<FConsoleVariablesEditorAssetSaveData> SavedCommands;
+};
diff --git a/Engine/Plugins/Editor/ConsoleVariablesEditor/Source/ConsoleVariablesEditorRuntime/Public/ConsoleVariablesEditorRuntimeLog.h b/Engine/Plugins/Editor/ConsoleVariablesEditor/Source/ConsoleVariablesEditorRuntime/Public/ConsoleVariablesEditorRuntimeLog.h
new file mode 100644
--- /dev/null
+++ b/Engine/Plugins/Editor/ConsoleVariablesEditor/Source/ConsoleVariablesEditorRuntime/Public/ConsoleVariablesEditorRuntimeLog.h
@@ -0,0 +1,7 @@
+// Copyright Epic Games, Inc. All Rights Reserved.
+
+#pragma once
+
+#include "Logging/LogMacros.h"
+
+DECLARE_LOG_CATEGORY_CLASS(LogConsoleVariablesEditorRuntime, Log, All);
