[Backout] - CL20544648
Original CL Desc
-----------------------------------------------------------------
Console Variables Editor:

-Move ConsoleVariablesAsset from UncookedOnly module to separate Runtime module
-Make Console Variable Asset functions mostly BlueprintCallable or BlueprintPure
-Add option to execute all commands in an asset from context click in content browser

diff --git a/Engine/Plugins/Editor/ConsoleVariablesEditor/ConsoleVariables.uplugin b/Engine/Plugins/Editor/ConsoleVariablesEditor/ConsoleVariables.uplugin
--- a/Engine/Plugins/Editor/ConsoleVariablesEditor/ConsoleVariables.uplugin
+++ b/Engine/Plugins/Editor/ConsoleVariablesEditor/ConsoleVariables.uplugin
@@ -39,11 +39,6 @@
 			"Name": "ConsoleVariablesEditor",
 			"Type": "UncookedOnly",
 			"LoadingPhase": "Default"
-		},
-		{
-			"Name": "ConsoleVariablesEditorRuntime",
-			"Type": "Runtime",
-			"LoadingPhase": "Default"
 		}
 	]
 }
diff --git a/Engine/Plugins/Editor/ConsoleVariablesEditor/Source/ConsoleVariablesEditor/ConsoleVariablesEditor.build.cs b/Engine/Plugins/Editor/ConsoleVariablesEditor/Source/ConsoleVariablesEditor/ConsoleVariablesEditor.build.cs
--- a/Engine/Plugins/Editor/ConsoleVariablesEditor/Source/ConsoleVariablesEditor/ConsoleVariablesEditor.build.cs
+++ b/Engine/Plugins/Editor/ConsoleVariablesEditor/Source/ConsoleVariablesEditor/ConsoleVariablesEditor.build.cs
@@ -21,15 +21,14 @@ public class ConsoleVariablesEditor : ModuleRules
 		PrivateDependencyModuleNames.AddRange(
 			new string[]
 			{
-				"AssetRegistry",
-				"AssetTools",
-				"CoreUObject",
 				"Concert",
 				"ConcertSharedSlate",
 				"ConcertSyncClient",
 				"ConcertSyncCore",
 				"ConcertTransport",
-				"ConsoleVariablesEditorRuntime",
+				"AssetRegistry",
+				"AssetTools",
+				"CoreUObject",
 				"ContentBrowser",
 				"Engine",
 				"EditorStyle",
diff --git a/Engine/Plugins/Editor/ConsoleVariablesEditor/Source/ConsoleVariablesEditor/Private/AssetTypeActions/AssetTypeActions_ConsoleVariables.cpp b/Engine/Plugins/Editor/ConsoleVariablesEditor/Source/ConsoleVariablesEditor/Private/AssetTypeActions/AssetTypeActions_ConsoleVariables.cpp
--- a/Engine/Plugins/Editor/ConsoleVariablesEditor/Source/ConsoleVariablesEditor/Private/AssetTypeActions/AssetTypeActions_ConsoleVariables.cpp
+++ b/Engine/Plugins/Editor/ConsoleVariablesEditor/Source/ConsoleVariablesEditor/Private/AssetTypeActions/AssetTypeActions_ConsoleVariables.cpp
@@ -3,7 +3,6 @@
 #include "AssetTypeActions_ConsoleVariables.h"
 
 #include "ConsoleVariablesAsset.h"
-#include "ConsoleVariablesEditorStyle.h"
 #include "ConsoleVariablesEditor/Private/ConsoleVariablesEditorModule.h"
 
 #include "Framework/MultiBox/MultiBoxBuilder.h"
@@ -25,81 +24,21 @@ void FAssetTypeActions_ConsoleVariables::GetActions(const TArray<UObject*>& InOb
 	FAssetTypeActions_Base::GetActions(InObjects, MenuBuilder);
 
 	const TArray<TWeakObjectPtr<UConsoleVariablesAsset>> ConsoleVariableAssets = GetTypedWeakObjectPtrs<UConsoleVariablesAsset>(InObjects);
-	const int32 ConsoleVariablesAssetCount = ConsoleVariableAssets.Num();
 
 	MenuBuilder.AddMenuEntry(
 		LOCTEXT("AssetTypeActions_OpenVariableCollection", "Open Variable Collection in Editor"),
 		LOCTEXT("AssetTypeActions_OpenVariableCollectionToolTip", "Open this console variable collection in the Console Variables Editor. Select only one asset at a time."),
 		FSlateIcon(FAppStyle::Get().GetStyleSetName(), "SystemWideCommands.SummonOpenAssetDialog"),
 		FUIAction(
-			FExecuteAction::CreateLambda([this, InObjects] {
+			FExecuteAction::CreateLambda([=] {
 					OpenAssetEditor(InObjects, TSharedPtr<IToolkitHost>());
 			}),
-			FCanExecuteAction::CreateLambda([ConsoleVariablesAssetCount] {
+			FCanExecuteAction::CreateLambda([=] {
 					// We only want to open a single Variable Collection asset at a time, so let's ensure
 					// the number of selected assets is exactly one.
-					return ConsoleVariablesAssetCount == 1;
-				}
+					return ConsoleVariableAssets.Num() == 1;
+				})
 			)
-		)
-	);
-	
-	MenuBuilder.AddMenuEntry(
-		LOCTEXT("AssetTypeActions_ExecuteVariableCollection", "Execute Variable Collection"),
-		LOCTEXT("AssetTypeActions_ExecuteVariableCollectionToolTip", "Executes all commands and variables saved to the selected assets."),
-		FSlateIcon(FConsoleVariablesEditorStyle::Get().GetStyleSetName(), "ConsoleVariables.ToolbarButton.Small"),
-		FUIAction(
-			FExecuteAction::CreateLambda([ConsoleVariableAssets] {
-				UWorld* CurrentWorld = nullptr;
-				check (GIsEditor);
-				{
-					CurrentWorld = GEditor->GetEditorWorldContext().World();
-				}
-				
-				for (TWeakObjectPtr<UConsoleVariablesAsset> Asset : ConsoleVariableAssets)
-				{
-					if (Asset.IsValid())
-					{
-						Asset->ExecuteSavedCommands(CurrentWorld);
-					}
-				}
-			}),
-			FCanExecuteAction::CreateLambda([ConsoleVariablesAssetCount] {
-					// We only want to open a single Variable Collection asset at a time, so let's ensure
-					// the number of selected assets is exactly one.
-					return ConsoleVariablesAssetCount > 0;
-				}
-			)
-		)
-	);
-
-	MenuBuilder.AddMenuEntry(
-		LOCTEXT("AssetTypeActions_ExecuteOnlyCheckedInVariableCollection", "Execute Variable Collection (Only Checked)"),
-		LOCTEXT("AssetTypeActions_ExecuteOnlyCheckedInVariableCollectionToolTip", "Executes commands and variables with a Checked checkstate saved to the selected assets. Checkstates can be edited in the Console Variables Editor UI."),
-		FSlateIcon(FConsoleVariablesEditorStyle::Get().GetStyleSetName(), "ConsoleVariables.ToolbarButton.Small"),
-		FUIAction(
-			FExecuteAction::CreateLambda([ConsoleVariableAssets] {
-				UWorld* CurrentWorld = nullptr;
-				check (GIsEditor)
-				{
-					CurrentWorld = GEditor->GetEditorWorldContext().World();
-				}
-				
-				for (TWeakObjectPtr<UConsoleVariablesAsset> Asset : ConsoleVariableAssets)
-				{
-					if (Asset.IsValid())
-					{
-						Asset->ExecuteSavedCommands(CurrentWorld, true);
-					}
-				}
-			}),
-			FCanExecuteAction::CreateLambda([ConsoleVariablesAssetCount] {
-					// We only want to open a single Variable Collection asset at a time, so let's ensure
-					// the number of selected assets is exactly one.
-					return ConsoleVariablesAssetCount > 0;
-				}
-			)
-		)
 	);
 }
 
diff --git a/Engine/Plugins/Editor/ConsoleVariablesEditor/Source/ConsoleVariablesEditor/Private/ConsoleVariablesAsset.cpp b/Engine/Plugins/Editor/ConsoleVariablesEditor/Source/ConsoleVariablesEditor/Private/ConsoleVariablesAsset.cpp
new file mode 100644
--- /dev/null
+++ b/Engine/Plugins/Editor/ConsoleVariablesEditor/Source/ConsoleVariablesEditor/Private/ConsoleVariablesAsset.cpp
@@ -0,0 +1,88 @@
+// Copyright Epic Games, Inc. All Rights Reserved.
+
+#include "ConsoleVariablesAsset.h"
+
+#include "ConsoleVariablesEditorCommandInfo.h"
+#include "ConsoleVariablesEditorLog.h"
+#include "ConsoleVariablesEditorModule.h"
+
+void UConsoleVariablesAsset::SetVariableCollectionDescription(const FString& InVariableCollectionDescription)
+{
+	VariableCollectionDescription = InVariableCollectionDescription;
+}
+
+void UConsoleVariablesAsset::ReplaceSavedCommands(const TArray<FConsoleVariablesEditorAssetSaveData>& Replacement)
+{
+	SavedCommands = Replacement;
+}
+
+bool UConsoleVariablesAsset::FindSavedDataByCommandString(
+	const FString InCommandString, FConsoleVariablesEditorAssetSaveData& OutValue, const ESearchCase::Type SearchCase) const
+{
+	for (const FConsoleVariablesEditorAssetSaveData& Command : SavedCommands)
+	{
+		if (Command.CommandName.TrimStartAndEnd().Equals(InCommandString.TrimStartAndEnd(), SearchCase))
+		{
+			OutValue = Command;
+			return true;
+		}
+	}
+	
+	return false;
+}
+
+void UConsoleVariablesAsset::AddOrSetConsoleObjectSavedData(const FConsoleVariablesEditorAssetSaveData& InData)
+{
+	FConsoleVariablesEditorModule& ConsoleVariablesEditorModule = FConsoleVariablesEditorModule::Get();
+	UE_LOG(LogConsoleVariablesEditor, VeryVerbose, TEXT("%hs: Adding %s to editable asset"),
+		__FUNCTION__, *InData.CommandName);
+	if (const TWeakPtr<FConsoleVariablesEditorCommandInfo> CommandInfo =
+			ConsoleVariablesEditorModule.FindCommandInfoByName(InData.CommandName);
+		CommandInfo.IsValid())
+	{
+		const IConsoleVariable* AsVariable = CommandInfo.Pin()->GetConsoleVariablePtr();
+		
+		if (AsVariable && AsVariable->TestFlags(ECVF_RenderThreadSafe))
+		{
+			UE_LOG(LogConsoleVariablesEditor, Verbose,
+				TEXT("The console variable named %s is flagged as ECVF_RenderThreadSafe. The value on the render thread will lag behind the value on the main thread by one frame if r.OneFrameThreadLag is 1."),
+				*InData.CommandName
+			);
+		}
+	}
+	
+	RemoveConsoleVariable(InData.CommandName);
+	SavedCommands.Add(InData);
+	// Make a copy because this array can change during iteration
+	TArray<FConsoleVariablesEditorAssetSaveData> SavedCommandsLocal = SavedCommands; 
+	for (int32 CommandIndex = 0; CommandIndex < SavedCommandsLocal.Num(); CommandIndex++)
+	{
+		FConsoleVariablesEditorAssetSaveData& SavedCommand = SavedCommandsLocal[CommandIndex];
+		UE_LOG(LogConsoleVariablesEditor, VeryVerbose, TEXT("%hs: Command named '%s' at Index %i"),
+			__FUNCTION__, *SavedCommand.CommandName, CommandIndex);
+	}
+}
+
+bool UConsoleVariablesAsset::RemoveConsoleVariable(const FString InCommandString)
+{
+	FConsoleVariablesEditorAssetSaveData ExistingData;
+	int32 RemoveCount = 0;
+	while (FindSavedDataByCommandString(InCommandString, ExistingData, ESearchCase::IgnoreCase))
+	{
+		UE_LOG(LogConsoleVariablesEditor, VeryVerbose, TEXT("%hs: Removing %s from editable asset"),
+			__FUNCTION__, *InCommandString);
+		
+		if (SavedCommands.Remove(ExistingData) > 0)
+		{
+			RemoveCount++;
+		}
+	}
+	
+	return RemoveCount > 0;
+}
+
+void UConsoleVariablesAsset::CopyFrom(const UConsoleVariablesAsset* InAssetToCopy)
+{
+	VariableCollectionDescription = InAssetToCopy->GetVariableCollectionDescription();
+	SavedCommands = InAssetToCopy->GetSavedCommands();
+}
diff --git a/Engine/Plugins/Editor/ConsoleVariablesEditor/Source/ConsoleVariablesEditor/Public/ConsoleVariablesAsset.h b/Engine/Plugins/Editor/ConsoleVariablesEditor/Source/ConsoleVariablesEditor/Public/ConsoleVariablesAsset.h
new file mode 100644
--- /dev/null
+++ b/Engine/Plugins/Editor/ConsoleVariablesEditor/Source/ConsoleVariablesEditor/Public/ConsoleVariablesAsset.h
@@ -0,0 +1,85 @@
+// Copyright Epic Games, Inc. All Rights Reserved.
+
+#pragma once
+
+#include "CoreMinimal.h"
+#include "Widgets/Input/SCheckBox.h"
+
+#include "ConsoleVariablesAsset.generated.h"
+
+/** Data that will be serialized with this asset */
+USTRUCT()
+struct FConsoleVariablesEditorAssetSaveData
+{
+	GENERATED_BODY()
+
+	FORCEINLINE bool operator==(const FConsoleVariablesEditorAssetSaveData& Comparator) const
+	{
+		return CommandName.Equals(Comparator.CommandName);
+	}
+	
+	UPROPERTY()
+	FString CommandName;
+
+	UPROPERTY()
+	FString CommandValueAsString;
+
+	UPROPERTY()
+	/** If Undetermined, we can assume this data was not previously saved */
+	ECheckBoxState CheckedState = ECheckBoxState::Undetermined;
+};
+
+/** An asset used to track collections of console variables that can be recalled and edited using the Console Variables Editor. */
+UCLASS(BlueprintType)
+class CONSOLEVARIABLESEDITOR_API UConsoleVariablesAsset : public UObject
+{
+	GENERATED_BODY()
+public:
+
+	/** Sets a description for this variable collection. */
+	UFUNCTION(BlueprintCallable, Category = "Console Variables Editor")
+	void SetVariableCollectionDescription(const FString& InVariableCollectionDescription);
+	UFUNCTION(BlueprintPure, Category = "Console Variables Editor")
+	FString GetVariableCollectionDescription() const
+	{
+		return VariableCollectionDescription;
+	}
+
+	/** Returns the saved list of console variable information such as the variable name, the type and the value of the vriable at the time the asset was saved. */
+	const TArray<FConsoleVariablesEditorAssetSaveData>& GetSavedCommands() const
+	{
+		return SavedCommands;
+	}
+
+	/** Completely replaces the saved data with new saved data */
+	void ReplaceSavedCommands(const TArray<FConsoleVariablesEditorAssetSaveData>& Replacement);
+
+	/** Returns how many console variables are serialized in this asset */
+	int32 GetSavedCommandsCount() const
+	{
+		return GetSavedCommands().Num();
+	}
+
+	/** Outputs the FConsoleVariablesEditorCommandInfo matching InCommand. Returns whether a match was found. Case sensitive. */
+	bool FindSavedDataByCommandString(
+		const FString InCommandString, FConsoleVariablesEditorAssetSaveData& OutValue, const ESearchCase::Type SearchCase = ESearchCase::IgnoreCase) const;
+
+	/** Set the value of a saved console variable if the name matches; add a new console variable to the list if a match is not found. */
+	void AddOrSetConsoleObjectSavedData(const FConsoleVariablesEditorAssetSaveData& InData);
+
+	/** Returns true if the element was found and successfully removed. */
+	bool RemoveConsoleVariable(const FString InCommandString);
+
+	/** Copy data from input asset to this asset */
+	void CopyFrom(const UConsoleVariablesAsset* InAssetToCopy);
+	
+private:
+	
+	/* User defined description of the variable collection */
+	UPROPERTY(AssetRegistrySearchable, BlueprintGetter = "GetVariableCollectionDescription", EditAnywhere, Category = "Console Variables Editor")
+	FString VariableCollectionDescription;
+
+	/** A saved list of console variable information such as the variable name, the type and the value of the variable at the time the asset was saved. */
+	UPROPERTY()
+	TArray<FConsoleVariablesEditorAssetSaveData> SavedCommands;
+};
diff --git a/Engine/Plugins/Editor/ConsoleVariablesEditor/Source/ConsoleVariablesEditorRuntime/ConsoleVariablesEditorRuntime.build.cs b/Engine/Plugins/Editor/ConsoleVariablesEditor/Source/ConsoleVariablesEditorRuntime/ConsoleVariablesEditorRuntime.build.cs
deleted file mode 100644
diff --git a/Engine/Plugins/Editor/ConsoleVariablesEditor/Source/ConsoleVariablesEditorRuntime/Private/ConsoleVariablesAsset.cpp b/Engine/Plugins/Editor/ConsoleVariablesEditor/Source/ConsoleVariablesEditorRuntime/Private/ConsoleVariablesAsset.cpp
deleted file mode 100644
diff --git a/Engine/Plugins/Editor/ConsoleVariablesEditor/Source/ConsoleVariablesEditorRuntime/Private/ConsoleVariablesEditorRuntimeModule.cpp b/Engine/Plugins/Editor/ConsoleVariablesEditor/Source/ConsoleVariablesEditorRuntime/Private/ConsoleVariablesEditorRuntimeModule.cpp
deleted file mode 100644
diff --git a/Engine/Plugins/Editor/ConsoleVariablesEditor/Source/ConsoleVariablesEditorRuntime/Private/ConsoleVariablesEditorRuntimeModule.h b/Engine/Plugins/Editor/ConsoleVariablesEditor/Source/ConsoleVariablesEditorRuntime/Private/ConsoleVariablesEditorRuntimeModule.h
deleted file mode 100644
diff --git a/Engine/Plugins/Editor/ConsoleVariablesEditor/Source/ConsoleVariablesEditorRuntime/Public/ConsoleVariablesAsset.h b/Engine/Plugins/Editor/ConsoleVariablesEditor/Source/ConsoleVariablesEditorRuntime/Public/ConsoleVariablesAsset.h
deleted file mode 100644
diff --git a/Engine/Plugins/Editor/ConsoleVariablesEditor/Source/ConsoleVariablesEditorRuntime/Public/ConsoleVariablesEditorRuntimeLog.h b/Engine/Plugins/Editor/ConsoleVariablesEditor/Source/ConsoleVariablesEditorRuntime/Public/ConsoleVariablesEditorRuntimeLog.h
deleted file mode 100644
