Replace 'Master' with 'Central'

diff --git a/Engine/Plugins/Editor/ConsoleVariablesEditor/Source/ConsoleVariablesEditor/Private/ConsoleVariablesEditorModule.cpp b/Engine/Plugins/Editor/ConsoleVariablesEditor/Source/ConsoleVariablesEditor/Private/ConsoleVariablesEditorModule.cpp
--- a/Engine/Plugins/Editor/ConsoleVariablesEditor/Source/ConsoleVariablesEditor/Private/ConsoleVariablesEditorModule.cpp
+++ b/Engine/Plugins/Editor/ConsoleVariablesEditor/Source/ConsoleVariablesEditor/Private/ConsoleVariablesEditorModule.cpp
@@ -49,7 +49,7 @@ void FConsoleVariablesEditorModule::ShutdownModule()
 	
 	MainPanel.Reset();
 
-	ConsoleObjectsMasterReference.Empty();
+	ConsoleObjectsMainReference.Empty();
 
 	// Unregister project settings
 	ISettingsModule& SettingsModule = FModuleManager::LoadModuleChecked<ISettingsModule>("Settings");
@@ -114,9 +114,9 @@ void FConsoleVariablesEditorModule::OpenConsoleVariablesDialogWithAssetSelected(
 
 void FConsoleVariablesEditorModule::QueryAndBeginTrackingConsoleVariables()
 {
-	const int32 VariableCount = ConsoleObjectsMasterReference.Num();
+	const int32 VariableCount = ConsoleObjectsMainReference.Num();
 	
-	ConsoleObjectsMasterReference.Empty(VariableCount);
+	ConsoleObjectsMainReference.Empty(VariableCount);
 	
 	IConsoleManager::Get().ForEachConsoleObjectThatStartsWith(FConsoleObjectVisitor::CreateLambda(
 		[this] (const TCHAR* Key, IConsoleObject* ConsoleObject)
@@ -139,7 +139,7 @@ void FConsoleVariablesEditorModule::QueryAndBeginTrackingConsoleVariables()
 					this, &FConsoleVariablesEditorModule::OnConsoleVariableChanged);
 			}
 			
-			AddConsoleObjectCommandInfoToMasterReference(Info);
+			AddConsoleObjectCommandInfoToMainReference(Info);
 		}),
 		TEXT(""));
 }
@@ -147,7 +147,7 @@ void FConsoleVariablesEditorModule::QueryAndBeginTrackingConsoleVariables()
 TWeakPtr<FConsoleVariablesEditorCommandInfo> FConsoleVariablesEditorModule::FindCommandInfoByName(const FString& NameToSearch, ESearchCase::Type InSearchCase)
 {
 	TSharedPtr<FConsoleVariablesEditorCommandInfo>* Match = Algo::FindByPredicate(
-		ConsoleObjectsMasterReference,
+		ConsoleObjectsMainReference,
 		[&NameToSearch, InSearchCase](const TSharedPtr<FConsoleVariablesEditorCommandInfo> Comparator)
 		{
 			return Comparator->Command.Equals(NameToSearch, InSearchCase);
@@ -161,7 +161,7 @@ TArray<TWeakPtr<FConsoleVariablesEditorCommandInfo>> FConsoleVariablesEditorModu
 {
 	TArray<TWeakPtr<FConsoleVariablesEditorCommandInfo>> ReturnValue;
 	
-	for (const TSharedPtr<FConsoleVariablesEditorCommandInfo>& CommandInfo : ConsoleObjectsMasterReference)
+	for (const TSharedPtr<FConsoleVariablesEditorCommandInfo>& CommandInfo : ConsoleObjectsMainReference)
 	{
 		FString CommandSearchableText = CommandInfo->Command + " " + CommandInfo->GetHelpText() + " " + CommandInfo->GetSourceAsText().ToString();
 		// Match any
@@ -198,7 +198,7 @@ TWeakPtr<FConsoleVariablesEditorCommandInfo> FConsoleVariablesEditorModule::Find
 	IConsoleObject* InConsoleObjectReference)
 {
 	TSharedPtr<FConsoleVariablesEditorCommandInfo>* Match = Algo::FindByPredicate(
-	ConsoleObjectsMasterReference,
+	ConsoleObjectsMainReference,
 	[InConsoleObjectReference](const TSharedPtr<FConsoleVariablesEditorCommandInfo> Comparator)
 	{
 		return Comparator->GetConsoleObjectPtr() == InConsoleObjectReference;
@@ -431,7 +431,7 @@ void FConsoleVariablesEditorModule::OnDetectConsoleObjectUnregistered(FString Co
 	if (const TWeakPtr<FConsoleVariablesEditorCommandInfo> CommandInfo =
 		FindCommandInfoByName(CommandName); CommandInfo.IsValid())
 	{
-		ConsoleObjectsMasterReference.Remove(CommandInfo.Pin());
+		ConsoleObjectsMainReference.Remove(CommandInfo.Pin());
 	}
 }
 
diff --git a/Engine/Plugins/Editor/ConsoleVariablesEditor/Source/ConsoleVariablesEditor/Private/ConsoleVariablesEditorModule.h b/Engine/Plugins/Editor/ConsoleVariablesEditor/Source/ConsoleVariablesEditor/Private/ConsoleVariablesEditorModule.h
--- a/Engine/Plugins/Editor/ConsoleVariablesEditor/Source/ConsoleVariablesEditor/Private/ConsoleVariablesEditorModule.h
+++ b/Engine/Plugins/Editor/ConsoleVariablesEditor/Source/ConsoleVariablesEditor/Private/ConsoleVariablesEditorModule.h
@@ -38,9 +38,9 @@ public:
 	/** Find all console variables and cache their startup values */
 	void QueryAndBeginTrackingConsoleVariables();
 
-	void AddConsoleObjectCommandInfoToMasterReference(TSharedRef<FConsoleVariablesEditorCommandInfo> InCommandInfo)
+	void AddConsoleObjectCommandInfoToMainReference(TSharedRef<FConsoleVariablesEditorCommandInfo> InCommandInfo)
 	{
-		ConsoleObjectsMasterReference.Add(InCommandInfo);
+		ConsoleObjectsMainReference.Add(InCommandInfo);
 	}
 
 	/** Find a tracked console variable by the command string with optional case sensitivity. */
@@ -119,7 +119,7 @@ private:
 	TObjectPtr<UConsoleVariablesAsset> EditingGlobalSearchAsset = nullptr;
 
 	/** All tracked variables and their default, startup, and current values */
-	TArray<TSharedPtr<FConsoleVariablesEditorCommandInfo>> ConsoleObjectsMasterReference;
+	TArray<TSharedPtr<FConsoleVariablesEditorCommandInfo>> ConsoleObjectsMainReference;
 
 	/**
 	 * A list of CommandName/ValueAsString pairs of commands received from MU.
diff --git a/Engine/Plugins/Editor/ConsoleVariablesEditor/Source/ConsoleVariablesEditor/Private/Views/List/SConsoleVariablesEditorList.cpp b/Engine/Plugins/Editor/ConsoleVariablesEditor/Source/ConsoleVariablesEditor/Private/Views/List/SConsoleVariablesEditorList.cpp
--- a/Engine/Plugins/Editor/ConsoleVariablesEditor/Source/ConsoleVariablesEditor/Private/Views/List/SConsoleVariablesEditorList.cpp
+++ b/Engine/Plugins/Editor/ConsoleVariablesEditor/Source/ConsoleVariablesEditor/Private/Views/List/SConsoleVariablesEditorList.cpp
@@ -1113,7 +1113,7 @@ void SConsoleVariablesEditorList::GenerateTreeView(const bool bSkipExecutionOfCa
 		{
 			CommandInfo = MakeShared<FConsoleVariablesEditorCommandInfo>(SavedCommand.CommandName);
 
-			FConsoleVariablesEditorModule::Get().AddConsoleObjectCommandInfoToMasterReference(
+			FConsoleVariablesEditorModule::Get().AddConsoleObjectCommandInfoToMainReference(
 				CommandInfo.ToSharedRef());
 		}
 		
