[NO NON-PII COMMIT MESSAGE]

diff --git a/Engine/Plugins/Editor/ConsoleVariablesEditor/Source/ConsoleVariablesEditor/Private/ConsoleVariablesEditorModule.cpp b/Engine/Plugins/Editor/ConsoleVariablesEditor/Source/ConsoleVariablesEditor/Private/ConsoleVariablesEditorModule.cpp
--- a/Engine/Plugins/Editor/ConsoleVariablesEditor/Source/ConsoleVariablesEditor/Private/ConsoleVariablesEditorModule.cpp
+++ b/Engine/Plugins/Editor/ConsoleVariablesEditor/Source/ConsoleVariablesEditor/Private/ConsoleVariablesEditorModule.cpp
@@ -258,25 +258,29 @@ bool FConsoleVariablesEditorModule::PopulateGlobalSearchAssetWithVariablesMatchi
 	return EditingGlobalSearchAsset->GetSavedCommandsCount() > 0;
 }
 
-void FConsoleVariablesEditorModule::SendMultiUserConsoleVariableChange(const FString& InVariableName, const FString& InValueAsString) const
+void FConsoleVariablesEditorModule::SendMultiUserConsoleVariableChange(const FString& InVariableName, const FString& InValueAsString)
 {
-	if (GEditor && GEditor->IsPlaySessionInProgress())
+	if (GEditor && GEditor->IsPlaySessionInProgress() && !bHaveWarnedAboutPIE)
 	{
-		UE_LOG(LogConsoleVariablesEditor, Warning, TEXT("%hs: Play In Editor is about to start or has started; Multi-User Cvar sync is suspended duirng PIE."), __FUNCTION__);
+		UE_LOG(LogConsoleVariablesEditor, Warning, TEXT("%hs: Play In Editor is about to start or has started; Multi-User Cvar sync is suspended during PIE."), __FUNCTION__);
+		bHaveWarnedAboutPIE = true;
 		return;
 	}
-	
+
+	bHaveWarnedAboutPIE = false;
 	MainPanel->GetMultiUserManager().SendConsoleVariableChange(InVariableName, InValueAsString);
 }
 
 void FConsoleVariablesEditorModule::OnRemoteCvarChanged(const FString InName, const FString InValue)
 {
-	if (GEditor && GEditor->IsPlaySessionInProgress())
+	if (GEditor && GEditor->IsPlaySessionInProgress() && !bHaveWarnedAboutPIE)
 	{
-		UE_LOG(LogConsoleVariablesEditor, Warning, TEXT("%hs: Play In Editor is about to start or has started; Multi-User Cvar sync is suspended duirng PIE."), __FUNCTION__);
+		UE_LOG(LogConsoleVariablesEditor, Warning, TEXT("%hs: Play In Editor is about to start or has started; Multi-User Cvar sync is suspended during PIE."), __FUNCTION__);
+		bHaveWarnedAboutPIE = true;
 		return;
 	}
-	
+
+	bHaveWarnedAboutPIE = false;
 	UE_LOG(LogConsoleVariablesEditor, VeryVerbose, TEXT("Remote set console variable %s = %s"), *InName, *InValue);
 
 	CommandsRecentlyReceivedFromMultiUser.Add(InName, InValue);
diff --git a/Engine/Plugins/Editor/ConsoleVariablesEditor/Source/ConsoleVariablesEditor/Private/ConsoleVariablesEditorModule.h b/Engine/Plugins/Editor/ConsoleVariablesEditor/Source/ConsoleVariablesEditor/Private/ConsoleVariablesEditorModule.h
--- a/Engine/Plugins/Editor/ConsoleVariablesEditor/Source/ConsoleVariablesEditor/Private/ConsoleVariablesEditorModule.h
+++ b/Engine/Plugins/Editor/ConsoleVariablesEditor/Source/ConsoleVariablesEditor/Private/ConsoleVariablesEditorModule.h
@@ -84,7 +84,7 @@ public:
 	/** Fills Global Search Asset's Saved Commands with variables matching the specified query. Returns false if no matches were found. */
 	bool PopulateGlobalSearchAssetWithVariablesMatchingTokens(const TArray<FString>& InTokens);
 
-	void SendMultiUserConsoleVariableChange(const FString& InVariableName, const FString& InValueAsString) const;
+	void SendMultiUserConsoleVariableChange(const FString& InVariableName, const FString& InValueAsString);
 	void OnRemoteCvarChanged(const FString InName, const FString InValue);
 
 	virtual void AddReferencedObjects( FReferenceCollector& Collector )  override;
@@ -128,4 +128,7 @@ private:
 	 * if the command and value match an item in this map, the change will not be propagated to other nodes.
 	 */
 	TMap<FString, FString> CommandsRecentlyReceivedFromMultiUser;
+
+	/* Have we warned the user about PIE and Console Variable Editor */
+	bool bHaveWarnedAboutPIE = false;
 };
