Add a core redirect for the ConsoleVariableAsset

diff --git a/Engine/Plugins/Editor/ConsoleVariablesEditor/Config/BaseConsoleVariables.ini b/Engine/Plugins/Editor/ConsoleVariablesEditor/Config/BaseConsoleVariables.ini
new file mode 100644
--- /dev/null
+++ b/Engine/Plugins/Editor/ConsoleVariablesEditor/Config/BaseConsoleVariables.ini
@@ -0,0 +1,2 @@
+[CoreRedirects]
++ClassRedirects=(OldName="/Script/ConsoleVariablesEditor.ConsoleVariablesAsset",NewName="/Script/ConsoleVariablesEditorRuntime.ConsoleVariablesAsset")
\ No newline at end of file
