Clean-up Multi-user server messages - multiple messages were sent creating ping / pong scenario.  Added a enum to detect modifies, adds and removes from CVE asset.

diff --git a/Engine/Plugins/Editor/ConsoleVariablesEditor/Source/ConsoleVariablesEditor/Private/ConsoleVariablesEditorCommandInfo.cpp b/Engine/Plugins/Editor/ConsoleVariablesEditor/Source/ConsoleVariablesEditor/Private/ConsoleVariablesEditorCommandInfo.cpp
--- a/Engine/Plugins/Editor/ConsoleVariablesEditor/Source/ConsoleVariablesEditor/Private/ConsoleVariablesEditorCommandInfo.cpp
+++ b/Engine/Plugins/Editor/ConsoleVariablesEditor/Source/ConsoleVariablesEditor/Private/ConsoleVariablesEditorCommandInfo.cpp
@@ -26,7 +26,7 @@ void FConsoleVariablesEditorCommandInfo::ExecuteCommand(
 	}
 	if (bShouldTransactInConcert)
 	{
-		FConsoleVariablesEditorModule::Get().SendMultiUserConsoleVariableChange(Command, NewValueAsString);
+		FConsoleVariablesEditorModule::Get().SendMultiUserConsoleVariableChange(ERemoteCVarChangeType::Update, Command, NewValueAsString);
 	}
 }
 
diff --git a/Engine/Plugins/Editor/ConsoleVariablesEditor/Source/ConsoleVariablesEditor/Private/ConsoleVariablesEditorModule.cpp b/Engine/Plugins/Editor/ConsoleVariablesEditor/Source/ConsoleVariablesEditor/Private/ConsoleVariablesEditorModule.cpp
--- a/Engine/Plugins/Editor/ConsoleVariablesEditor/Source/ConsoleVariablesEditor/Private/ConsoleVariablesEditorModule.cpp
+++ b/Engine/Plugins/Editor/ConsoleVariablesEditor/Source/ConsoleVariablesEditor/Private/ConsoleVariablesEditorModule.cpp
@@ -7,6 +7,7 @@
 #include "ConsoleVariablesEditorLog.h"
 #include "ConsoleVariablesEditorProjectSettings.h"
 #include "ConsoleVariablesEditorStyle.h"
+#include "MultiUser/ConsoleVariableSync.h"
 #include "MultiUser/ConsoleVariableSyncData.h"
 #include "Views/MainPanel/ConsoleVariablesEditorMainPanel.h"
 
@@ -258,7 +259,7 @@ bool FConsoleVariablesEditorModule::PopulateGlobalSearchAssetWithVariablesMatchi
 	return EditingGlobalSearchAsset->GetSavedCommandsCount() > 0;
 }
 
-void FConsoleVariablesEditorModule::SendMultiUserConsoleVariableChange(const FString& InVariableName, const FString& InValueAsString)
+void FConsoleVariablesEditorModule::SendMultiUserConsoleVariableChange(ERemoteCVarChangeType InChangeType, const FString& InVariableName, const FString& InValueAsString)
 {
 	if (GEditor && GEditor->IsPlaySessionInProgress() && !bHaveWarnedAboutPIE)
 	{
@@ -268,10 +269,10 @@ void FConsoleVariablesEditorModule::SendMultiUserConsoleVariableChange(const FSt
 	}
 
 	bHaveWarnedAboutPIE = false;
-	MainPanel->GetMultiUserManager().SendConsoleVariableChange(InVariableName, InValueAsString);
+	MainPanel->GetMultiUserManager().SendConsoleVariableChange(InChangeType, InVariableName, InValueAsString);
 }
 
-void FConsoleVariablesEditorModule::OnRemoteCvarChanged(const FString InName, const FString InValue)
+void FConsoleVariablesEditorModule::OnRemoteCvarChanged(ERemoteCVarChangeType InChangeType, const FString InName, const FString InValue)
 {
 	if (GEditor && GEditor->IsPlaySessionInProgress() && !bHaveWarnedAboutPIE)
 	{
@@ -283,35 +284,31 @@ void FConsoleVariablesEditorModule::OnRemoteCvarChanged(const FString InName, co
 	bHaveWarnedAboutPIE = false;
 	UE_LOG(LogConsoleVariablesEditor, VeryVerbose, TEXT("Remote set console variable %s = %s"), *InName, *InValue);
 
-	CommandsRecentlyReceivedFromMultiUser.Add(InName, InValue);
-
 	if (GetDefault<UConcertCVarSynchronization>()->bSyncCVarTransactions)
 	{
-		bool bShouldExecute = true;
-		
-		if (const TWeakPtr<FConsoleVariablesEditorCommandInfo> CommandInfo =
-			FindCommandInfoByName(InName); CommandInfo.IsValid())
+		FScopeMultiUserReceiveCVar CVarChange(*InName, CommandsReceivedFromMultiUser);
+
+		GEngine->Exec(FConsoleVariablesEditorCommandInfo::GetCurrentWorld(),
+					  *FString::Printf(TEXT("%s %s"), *InName, *InValue));
+
+		if (InChangeType == ERemoteCVarChangeType::Remove)
 		{
-			bShouldExecute = CommandInfo.Pin()->IsCurrentValueDifferentFromInputValue(InValue);
+			EditingPresetAsset->RemoveConsoleVariable(InName);
 		}
-
-		if (bShouldExecute)
+		else
 		{
-			GEngine->Exec(FConsoleVariablesEditorCommandInfo::GetCurrentWorld(),
-				*FString::Printf(TEXT("%s %s"), *InName, *InValue));
-			
 			EditingPresetAsset->AddOrSetConsoleObjectSavedData(
-				{
-					InName,
-					InValue,
-					ECheckBoxState::Checked
-				}
-			);
-
-			if (MainPanel->GetEditorListMode() == FConsoleVariablesEditorList::EConsoleVariablesEditorListMode::Preset)
 			{
-				MainPanel->RebuildList();
+				InName,
+				InValue,
+				ECheckBoxState::Checked
 			}
+			);
+		}
+
+		if (MainPanel->GetEditorListMode() == FConsoleVariablesEditorList::EConsoleVariablesEditorListMode::Preset)
+		{
+			MainPanel->RebuildList();
 		}
 	}
 }
@@ -354,6 +351,11 @@ void FConsoleVariablesEditorModule::RegisterProjectSettings() const
 
 void FConsoleVariablesEditorModule::OnConsoleVariableChanged(IConsoleVariable* ChangedVariable)
 {
+	if (IsEngineExitRequested())
+	{
+		return;
+	}
+
 	check(EditingPresetAsset);
 
 	if (const TWeakPtr<FConsoleVariablesEditorCommandInfo> CommandInfo =
@@ -361,13 +363,13 @@ void FConsoleVariablesEditorModule::OnConsoleVariableChanged(IConsoleVariable* C
 	{
 		const TSharedPtr<FConsoleVariablesEditorCommandInfo>& PinnedCommand = CommandInfo.Pin();
 		const FString& Key = PinnedCommand->Command;
-		
+
 		FConsoleVariablesEditorAssetSaveData FoundData;
 		bool bIsVariableCurrentlyTracked = EditingPresetAsset->FindSavedDataByCommandString(Key, FoundData, ESearchCase::IgnoreCase);
 
 		const UConsoleVariablesEditorProjectSettings* Settings = GetDefault<UConsoleVariablesEditorProjectSettings>();
 		check(Settings);
-		
+
 		if (!bIsVariableCurrentlyTracked)
 		{
 			// If not yet tracked and we want to track variable changes from outside the dialogue,
@@ -393,30 +395,24 @@ void FConsoleVariablesEditorModule::OnConsoleVariableChanged(IConsoleVariable* C
 			}
 		}
 
+		if (MainPanel.IsValid())
+		{
+			MainPanel->RefreshList();
+		}
+
 		// If the variable is already tracked or was just added to tracking, run the following code
 		if (bIsVariableCurrentlyTracked)
 		{
-			if (MainPanel.IsValid())
-			{
-				MainPanel->RefreshList();
-			}
-			
 			/**
 			 * Here we check the map of recently received variables from other nodes
 			 * If the command is in the map and the value is similar, we won't send the value to other nodes
 			 * because we can assume that this value came from another node.
 			 * This prevents a feedback loop.
 			 */
-			if (const FString* MatchedValue = CommandsRecentlyReceivedFromMultiUser.Find(Key))
+			if (!CommandsReceivedFromMultiUser.Find(Key))
 			{
-				if (MatchedValue->Equals(ChangedVariable->GetString()))
-				{
-					CommandsRecentlyReceivedFromMultiUser.Remove(Key);
-					return;
-				}
+				SendMultiUserConsoleVariableChange(ERemoteCVarChangeType::Update, Key, ChangedVariable->GetString());
 			}
-				
-			SendMultiUserConsoleVariableChange(Key, ChangedVariable->GetString());
 		}
 	}
 }
diff --git a/Engine/Plugins/Editor/ConsoleVariablesEditor/Source/ConsoleVariablesEditor/Private/ConsoleVariablesEditorModule.h b/Engine/Plugins/Editor/ConsoleVariablesEditor/Source/ConsoleVariablesEditor/Private/ConsoleVariablesEditorModule.h
--- a/Engine/Plugins/Editor/ConsoleVariablesEditor/Source/ConsoleVariablesEditor/Private/ConsoleVariablesEditorModule.h
+++ b/Engine/Plugins/Editor/ConsoleVariablesEditor/Source/ConsoleVariablesEditor/Private/ConsoleVariablesEditorModule.h
@@ -84,13 +84,50 @@ public:
 	/** Fills Global Search Asset's Saved Commands with variables matching the specified query. Returns false if no matches were found. */
 	bool PopulateGlobalSearchAssetWithVariablesMatchingTokens(const TArray<FString>& InTokens);
 
-	void SendMultiUserConsoleVariableChange(const FString& InVariableName, const FString& InValueAsString);
-	void OnRemoteCvarChanged(const FString InName, const FString InValue);
+	void SendMultiUserConsoleVariableChange(ERemoteCVarChangeType InChangeType, const FString& InVariableName, const FString& InValueAsString);
+	void OnRemoteCvarChanged(ERemoteCVarChangeType InChangeType, FString InName, FString InValue);
 
 	virtual void AddReferencedObjects( FReferenceCollector& Collector )  override;
 	virtual FString GetReferencerName() const override;
 
 private:
+	/** Scoped based struct to track inbound cvars set by Multi-user. */
+	struct FScopeMultiUserReceiveCVar
+	{
+		FScopeMultiUserReceiveCVar() = delete;
+		FScopeMultiUserReceiveCVar(const FScopeMultiUserReceiveCVar&) = delete;
+		FScopeMultiUserReceiveCVar& operator=(const FScopeMultiUserReceiveCVar&) = delete;
+
+		FScopeMultiUserReceiveCVar(FString InCommand, TMap<FString, int32>& InTable)
+			: TrackedCommand(MoveTemp(InCommand)),
+			  InboundCommandTable(InTable)
+		{
+			if (int32* Value = InboundCommandTable.Find(TrackedCommand))
+			{
+				*Value = *Value + 1;
+			}
+			else
+			{
+				int32& NewValue = InboundCommandTable.Add(TrackedCommand);
+				NewValue = 0;
+			}
+		}
+
+		~FScopeMultiUserReceiveCVar()
+		{
+			if (int32* Value = InboundCommandTable.Find(TrackedCommand))
+			{
+				*Value = *Value - 1;
+				if (*Value == 0)
+				{
+					InboundCommandTable.Remove(TrackedCommand);
+				}
+			}
+		}
+
+		FString TrackedCommand;
+		TMap<FString, int32>& InboundCommandTable;
+	};
 
 	void OnFEngineLoopInitComplete();
 
@@ -107,7 +144,7 @@ private:
 	TSharedRef<SDockTab> SpawnMainPanelTab(const FSpawnTabArgs& Args);
 
 	static void OpenConsoleVariablesEditor();
-	
+
 	static const FName ConsoleVariablesToolkitPanelTabId;
 
 	/** Lives for as long as the module is loaded. */
@@ -122,12 +159,11 @@ private:
 	TArray<TSharedPtr<FConsoleVariablesEditorCommandInfo>> ConsoleObjectsMainReference;
 
 	/**
-	 * A list of CommandName/ValueAsString pairs of commands received from MU.
-	 * When this node receives a command and value from another node, it will be added to this list.
-	 * This list will be checked against detected local cvar changes and,
-	 * if the command and value match an item in this map, the change will not be propagated to other nodes.
+	 * A map of commands invoked by Multi-user.  Each command has a corresponding reference count.
+	 * When this node receives a command from another node, it will be added to this list.  This is prevent
+	 * a remote cvar change creating a ping/pong cvar updates effect between nodes.
 	 */
-	TMap<FString, FString> CommandsRecentlyReceivedFromMultiUser;
+	TMap<FString, int32> CommandsReceivedFromMultiUser;
 
 	/* Have we warned the user about PIE and Console Variable Editor */
 	bool bHaveWarnedAboutPIE = false;
diff --git a/Engine/Plugins/Editor/ConsoleVariablesEditor/Source/ConsoleVariablesEditor/Private/MultiUser/ConsoleVariableSync.cpp b/Engine/Plugins/Editor/ConsoleVariablesEditor/Source/ConsoleVariablesEditor/Private/MultiUser/ConsoleVariableSync.cpp
--- a/Engine/Plugins/Editor/ConsoleVariablesEditor/Source/ConsoleVariablesEditor/Private/MultiUser/ConsoleVariableSync.cpp
+++ b/Engine/Plugins/Editor/ConsoleVariablesEditor/Source/ConsoleVariablesEditor/Private/MultiUser/ConsoleVariableSync.cpp
@@ -306,16 +306,37 @@ struct FManagerImpl
 		}
 	}
 
+	static EConsoleVariableChangeType ConvertChangeType(ERemoteCVarChangeType InType)
+	{
+		switch (InType)
+		{
+		case ERemoteCVarChangeType::Add: return EConsoleVariableChangeType::Add;
+		case ERemoteCVarChangeType::Remove: return EConsoleVariableChangeType::Remove;
+		default:
+			return EConsoleVariableChangeType::Modify;
+		}
+	}
+
+	static ERemoteCVarChangeType ConvertChangeType(EConsoleVariableChangeType InType)
+	{
+		switch (InType)
+		{
+		case EConsoleVariableChangeType::Add: return ERemoteCVarChangeType::Add;
+		case EConsoleVariableChangeType::Remove: return ERemoteCVarChangeType::Remove;
+		default:
+			return ERemoteCVarChangeType::Update;
+		}
+	}
 	void HandleCVarSet(const FConcertSessionContext& Context, const FConcertSetConsoleVariableEvent& InEvent)
 	{
 		if (CanReceiveConsoleVariable())
 		{
-			RemoteCVarChanged.Broadcast(InEvent.Variable, InEvent.Value);
+			RemoteCVarChanged.Broadcast(ConvertChangeType(InEvent.ChangeType), InEvent.Variable, InEvent.Value);
 		}
 	}
 
-	/** 
-	 * Register a new multi-user session. 
+	/**
+	 * Register a new multi-user session.
 	 */
 	void Register(TSharedRef<IConcertClientSession> InSession)
 	{
@@ -393,7 +414,7 @@ struct FManagerImpl
 	 * Sends the given console variable name with the specified value to all connected endpoints.  It is up to the
 	 * endpoint to implement the change locally based on configured synchronization state.
 	 */
-	void SendConsoleVariableChange(FString InName, FString InValue)
+	void SendConsoleVariableChange(ERemoteCVarChangeType InChangeType, FString InName, FString InValue)
 	{
 		TSharedPtr<IConcertClientSession> Session = WeakSession.Pin();
 		if (!bIsEnabled || !Session.IsValid())
@@ -401,7 +422,7 @@ struct FManagerImpl
 			return;
 		}
 
-		FConcertSetConsoleVariableEvent OutEvent{MoveTemp(InName), MoveTemp(InValue)};
+		FConcertSetConsoleVariableEvent OutEvent{ConvertChangeType(InChangeType), MoveTemp(InName), MoveTemp(InValue)};
 		Session->SendCustomEvent(OutEvent, Session->GetSessionClientEndpointIds(),
 								 EConcertMessageFlags::ReliableOrdered | EConcertMessageFlags::UniqueId);
 	}
@@ -473,9 +494,9 @@ FOnMultiUserConnectionChange& FManager::OnConnectionChange()
 	return Implementation->ConnectionChanged;
 }
 
-void FManager::SendConsoleVariableChange(FString InName, FString InValue)
+void FManager::SendConsoleVariableChange(ERemoteCVarChangeType InChangeType, FString InName, FString InValue)
 {
-	Implementation->SendConsoleVariableChange(MoveTemp(InName), MoveTemp(InValue));
+	Implementation->SendConsoleVariableChange(InChangeType, MoveTemp(InName), MoveTemp(InValue));
 }
 
 void FManager::SendListItemCheckStateChange(FString InName, ECheckBoxState InCheckedState)
diff --git a/Engine/Plugins/Editor/ConsoleVariablesEditor/Source/ConsoleVariablesEditor/Private/MultiUser/ConsoleVariableSync.h b/Engine/Plugins/Editor/ConsoleVariablesEditor/Source/ConsoleVariablesEditor/Private/MultiUser/ConsoleVariableSync.h
--- a/Engine/Plugins/Editor/ConsoleVariablesEditor/Source/ConsoleVariablesEditor/Private/MultiUser/ConsoleVariableSync.h
+++ b/Engine/Plugins/Editor/ConsoleVariablesEditor/Source/ConsoleVariablesEditor/Private/MultiUser/ConsoleVariableSync.h
@@ -10,7 +10,14 @@
 
 #include "ConcertMessages.h"
 
-DECLARE_MULTICAST_DELEGATE_TwoParams(FOnRemoteCVarChange, FString, FString);
+enum class ERemoteCVarChangeType : uint8
+{
+	Update = 0,
+	Add,
+	Remove
+};
+
+DECLARE_MULTICAST_DELEGATE_ThreeParams(FOnRemoteCVarChange, ERemoteCVarChangeType, FString, FString);
 DECLARE_MULTICAST_DELEGATE_TwoParams(FOnRemoteListItemCheckStateChange, FString, ECheckBoxState);
 DECLARE_MULTICAST_DELEGATE_OneParam(FOnMultiUserConnectionChange, EConcertConnectionStatus);
 
@@ -40,8 +47,8 @@ struct FManager
 	FOnMultiUserConnectionChange& OnConnectionChange();
 
 	/** Sends the named console variable with value to all connected Multi-user clients */
-	void SendConsoleVariableChange(FString InName, FString InValue);
-	
+	void SendConsoleVariableChange(ERemoteCVarChangeType InChangeType, FString InName, FString InValue);
+
 	/** Sends the named console variable with value to all connected Multi-user clients */
 	void SendListItemCheckStateChange(FString InName, ECheckBoxState InCheckedState);
 
diff --git a/Engine/Plugins/Editor/ConsoleVariablesEditor/Source/ConsoleVariablesEditor/Private/MultiUser/ConsoleVariableSyncData.h b/Engine/Plugins/Editor/ConsoleVariablesEditor/Source/ConsoleVariablesEditor/Private/MultiUser/ConsoleVariableSyncData.h
--- a/Engine/Plugins/Editor/ConsoleVariablesEditor/Source/ConsoleVariablesEditor/Private/MultiUser/ConsoleVariableSyncData.h
+++ b/Engine/Plugins/Editor/ConsoleVariablesEditor/Source/ConsoleVariablesEditor/Private/MultiUser/ConsoleVariableSyncData.h
@@ -87,11 +87,28 @@ struct FConcertCVarSyncChangeEvent
 	bool bSyncCVarChanges = true;
 };
 
+UENUM()
+enum class EConsoleVariableChangeType : uint8
+{
+	/** The console variable was modified */
+	Modify = 0,
+
+	/** A console variable was added to the CVE asset */
+	Add,
+
+	/** A console variable was removed to the CVE asset */
+	Remove
+};
+
 USTRUCT()
 struct FConcertSetConsoleVariableEvent
 {
 	GENERATED_BODY()
 
+	/** */
+	UPROPERTY()
+	EConsoleVariableChangeType ChangeType = EConsoleVariableChangeType::Modify;
+
 	/** Console variable name to apply value change. */
 	UPROPERTY()
 	FString Variable;
diff --git a/Engine/Plugins/Editor/ConsoleVariablesEditor/Source/ConsoleVariablesEditor/Private/Views/List/ConsoleVariablesEditorListRow.cpp b/Engine/Plugins/Editor/ConsoleVariablesEditor/Source/ConsoleVariablesEditor/Private/Views/List/ConsoleVariablesEditorListRow.cpp
--- a/Engine/Plugins/Editor/ConsoleVariablesEditor/Source/ConsoleVariablesEditor/Private/Views/List/ConsoleVariablesEditorListRow.cpp
+++ b/Engine/Plugins/Editor/ConsoleVariablesEditor/Source/ConsoleVariablesEditor/Private/Views/List/ConsoleVariablesEditorListRow.cpp
@@ -7,6 +7,7 @@
 
 #include "Algo/AllOf.h"
 #include "Algo/AnyOf.h"
+#include "MultiUser/ConsoleVariableSync.h"
 #include "Views/List/ConsoleVariablesEditorList.h"
 #include "Views/List/SConsoleVariablesEditorList.h"
 
@@ -318,7 +319,7 @@ FReply FConsoleVariablesEditorListRow::OnActionButtonClicked()
 	FConsoleVariablesEditorModule& ConsoleVariablesEditorModule = FConsoleVariablesEditorModule::Get();
 
 	const FString& CommandName = GetCommandInfo().Pin()->Command;
-
+	const FString& StartupValue = GetCommandInfo().Pin()->StartupValueAsString;
 	const TObjectPtr<UConsoleVariablesAsset> EditableAsset = ConsoleVariablesEditorModule.GetPresetAsset();
 	check(EditableAsset);
 
@@ -338,6 +339,8 @@ FReply FConsoleVariablesEditorListRow::OnActionButtonClicked()
 		else
 		{
 			EditableAsset->RemoveConsoleVariable(CommandName);
+
+			ConsoleVariablesEditorModule.SendMultiUserConsoleVariableChange(ERemoteCVarChangeType::Remove, CommandName, StartupValue);
 		}
 	}
 	else
@@ -346,6 +349,7 @@ FReply FConsoleVariablesEditorListRow::OnActionButtonClicked()
 
 		EditableAsset->RemoveConsoleVariable(CommandName);
 
+		ConsoleVariablesEditorModule.SendMultiUserConsoleVariableChange(ERemoteCVarChangeType::Remove, CommandName, StartupValue);
 		ListViewPtr.Pin()->RebuildListWithListMode(ListViewPtr.Pin()->GetListModelPtr().Pin()->GetListMode());
 	}
 
diff --git a/Engine/Plugins/Editor/ConsoleVariablesEditor/Source/ConsoleVariablesEditor/Private/Views/List/SConsoleVariablesEditorListRow.cpp b/Engine/Plugins/Editor/ConsoleVariablesEditor/Source/ConsoleVariablesEditor/Private/Views/List/SConsoleVariablesEditorListRow.cpp
--- a/Engine/Plugins/Editor/ConsoleVariablesEditor/Source/ConsoleVariablesEditor/Private/Views/List/SConsoleVariablesEditorListRow.cpp
+++ b/Engine/Plugins/Editor/ConsoleVariablesEditor/Source/ConsoleVariablesEditor/Private/Views/List/SConsoleVariablesEditorListRow.cpp
@@ -5,6 +5,7 @@
 #include "ConsoleVariablesEditorList.h"
 #include "ConsoleVariablesEditorModule.h"
 #include "ConsoleVariablesEditorStyle.h"
+#include "MultiUser/ConsoleVariableSync.h"
 #include "SConsoleVariablesEditorListValueInput.h"
 #include "Views/Widgets/SConsoleVariablesEditorTooltipWidget.h"
 
@@ -441,18 +442,10 @@ void SConsoleVariablesEditorListRow::OnCheckboxStateChange(const ECheckBoxState
 			if (PinnedItem->IsRowChecked())
 			{
 				PinnedItem->GetCommandInfo().Pin()->ExecuteCommand(PinnedItem->GetCachedValue());
-
-				FConsoleVariablesEditorModule::Get().SendMultiUserConsoleVariableChange(
-					PinnedItem->GetCommandInfo().Pin()->Command, PinnedItem->GetCachedValue());
 			}
 			else
 			{
 				PinnedItem->ResetToStartupValueAndSource();
-
-				FConsoleVariablesEditorModule::Get().SendMultiUserConsoleVariableChange(
-					PinnedItem->GetCommandInfo().Pin()->Command,
-					PinnedItem->GetCommandInfo().Pin()->StartupValueAsString
-				);
 			}
 		}
 
diff --git a/Engine/Plugins/Editor/ConsoleVariablesEditor/Source/ConsoleVariablesEditor/Private/Views/MainPanel/ConsoleVariablesEditorMainPanel.cpp b/Engine/Plugins/Editor/ConsoleVariablesEditor/Source/ConsoleVariablesEditor/Private/Views/MainPanel/ConsoleVariablesEditorMainPanel.cpp
--- a/Engine/Plugins/Editor/ConsoleVariablesEditor/Source/ConsoleVariablesEditor/Private/Views/MainPanel/ConsoleVariablesEditorMainPanel.cpp
+++ b/Engine/Plugins/Editor/ConsoleVariablesEditor/Source/ConsoleVariablesEditor/Private/Views/MainPanel/ConsoleVariablesEditorMainPanel.cpp
@@ -7,6 +7,7 @@
 #include "ConsoleVariablesEditorCommandInfo.h"
 #include "ConsoleVariablesEditorLog.h"
 #include "ConsoleVariablesEditorModule.h"
+#include "MultiUser/ConsoleVariableSync.h"
 #include "Views/List/ConsoleVariablesEditorList.h"
 #include "Views/MainPanel/SConsoleVariablesEditorMainPanel.h"
 
@@ -62,13 +63,23 @@ void FConsoleVariablesEditorMainPanel::AddConsoleObjectToCurrentPreset(
 {
 	if (const TObjectPtr<UConsoleVariablesAsset> Asset = GetEditingAsset())
 	{
+		const FString TrimmedCommand = InConsoleCommand.TrimStartAndEnd();
+		FConsoleVariablesEditorAssetSaveData ExistingData;
+		const bool bHadCommandInAsset = Asset->FindSavedDataByCommandString(TrimmedCommand, ExistingData, ESearchCase::IgnoreCase);
+
 		Asset->AddOrSetConsoleObjectSavedData(
 			{
-				InConsoleCommand.TrimStartAndEnd(),
+				TrimmedCommand,
 				InValue,
 				ECheckBoxState::Checked
 			}
 		);
+		if (!bHadCommandInAsset)
+		{
+			GetConsoleVariablesModule().SendMultiUserConsoleVariableChange(
+				ERemoteCVarChangeType::Add,
+				TrimmedCommand, InValue);
+		}
 
 		if (GetEditorList().Pin()->GetListMode() == FConsoleVariablesEditorList::EConsoleVariablesEditorListMode::Preset)
 		{
@@ -309,11 +320,11 @@ void FConsoleVariablesEditorMainPanel::OnConnectionChanged(EConcertConnectionSta
 	}
 }
 
-void FConsoleVariablesEditorMainPanel::OnRemoteCvarChange(const FString InName, const FString InValue)
+void FConsoleVariablesEditorMainPanel::OnRemoteCvarChange(ERemoteCVarChangeType InChangeType, FString InName, FString InValue)
 {
 	FConsoleVariablesEditorModule& ConsoleVariablesEditorModule = FConsoleVariablesEditorModule::Get();
 
-	ConsoleVariablesEditorModule.OnRemoteCvarChanged(InName, InValue);
+	ConsoleVariablesEditorModule.OnRemoteCvarChanged(InChangeType, MoveTemp(InName), MoveTemp(InValue));
 }
 
 void FConsoleVariablesEditorMainPanel::OnRemoteListItemCheckStateChange(const FString InName, ECheckBoxState InCheckedState)
diff --git a/Engine/Plugins/Editor/ConsoleVariablesEditor/Source/ConsoleVariablesEditor/Private/Views/MainPanel/ConsoleVariablesEditorMainPanel.h b/Engine/Plugins/Editor/ConsoleVariablesEditor/Source/ConsoleVariablesEditor/Private/Views/MainPanel/ConsoleVariablesEditorMainPanel.h
--- a/Engine/Plugins/Editor/ConsoleVariablesEditor/Source/ConsoleVariablesEditor/Private/Views/MainPanel/ConsoleVariablesEditorMainPanel.h
+++ b/Engine/Plugins/Editor/ConsoleVariablesEditor/Source/ConsoleVariablesEditor/Private/Views/MainPanel/ConsoleVariablesEditorMainPanel.h
@@ -3,6 +3,7 @@
 #pragma once
 
 #include "ConsoleVariablesEditorModule.h"
+#include "MultiUser/ConsoleVariableSync.h"
 #include "Views/List/ConsoleVariablesEditorList.h"
 #include "Widgets/SWidget.h"
 
@@ -85,7 +86,7 @@ private:
 	TSharedPtr<FConsoleVariablesEditorList> EditorList;
 
 	static void OnConnectionChanged(EConcertConnectionStatus Status);
-	static void OnRemoteCvarChange(const FString InName, const FString InValue);
+	static void OnRemoteCvarChange(ERemoteCVarChangeType InChangeType, FString InName, FString InValue);
 	static void OnRemoteListItemCheckStateChange(const FString InName, ECheckBoxState InCheckedState);
 	
 	void SendListItemCheckStateChange(const FString& InName, ECheckBoxState InCheckedState);
