Console Variables Editor:

-Fix: Cvars edited in CVE do not generate output log events

diff --git a/Engine/Plugins/Editor/ConsoleVariablesEditor/Source/ConsoleVariablesEditor/Private/ConsoleVariablesEditorCommandInfo.cpp b/Engine/Plugins/Editor/ConsoleVariablesEditor/Source/ConsoleVariablesEditor/Private/ConsoleVariablesEditorCommandInfo.cpp
--- a/Engine/Plugins/Editor/ConsoleVariablesEditor/Source/ConsoleVariablesEditor/Private/ConsoleVariablesEditorCommandInfo.cpp
+++ b/Engine/Plugins/Editor/ConsoleVariablesEditor/Source/ConsoleVariablesEditor/Private/ConsoleVariablesEditorCommandInfo.cpp
@@ -1,9 +1,10 @@
 // Copyright Epic Games, Inc. All Rights Reserved.
 
 #include "ConsoleVariablesEditorCommandInfo.h"
-#include "Algo/Find.h"
-#include "CoreMinimal.h"
+
 #include "ConsoleVariablesEditorLog.h"
+
+#include "Algo/Find.h"
 #include "Editor.h"
 #include "Editor/EditorEngine.h"
 #include "Engine/GameEngine.h"
@@ -12,24 +13,43 @@
 #define LOCTEXT_NAMESPACE "ConsoleVariablesEditor"
 
 void FConsoleVariablesEditorCommandInfo::ExecuteCommand(
-	const FString& NewValueAsString, const bool bShouldTransactInConcert, const bool bSetInSession)
+	const FString& NewValueAsString, const bool bShouldTransactInConcert,
+	const bool bSetInSession, const bool bIsInteractiveChange)
 {
+	const FString FullCommand = FString::Printf(TEXT("%s %s"), *Command, *NewValueAsString).TrimStartAndEnd();
 	if (IConsoleVariable* AsVariable = GetConsoleVariablePtr())
 	{
 		AsVariable->Set(*NewValueAsString, GetSource());
 		bSetInCurrentSession = bSetInSession;
+
+		if (!bIsInteractiveChange)
+		{
+			PrintCommandOrVariable();
+		}
 	}
 	else
 	{
-		GEngine->Exec(GetCurrentWorld(),
-			*FString::Printf(TEXT("%s %s"), *Command, *NewValueAsString).TrimStartAndEnd());
+		GEngine->Exec(GetCurrentWorld(), *FullCommand);
 	}
-	if (bShouldTransactInConcert)
+	if (!bIsInteractiveChange && bShouldTransactInConcert)
 	{
 		FConsoleVariablesEditorModule::Get().SendMultiUserConsoleVariableChange(ERemoteCVarChangeType::Update, Command, NewValueAsString);
 	}
 }
 
+void FConsoleVariablesEditorCommandInfo::PrintCommandOrVariable()
+{
+	FString FullCommand = FString::Printf(TEXT("%s"), *Command).TrimStartAndEnd();
+	
+	FString Value;
+	if (GetCurrentValueAsString(Value))
+	{
+		FullCommand = FString::Printf(TEXT("%s %s"), *FullCommand, *Value).TrimStartAndEnd();
+	}
+
+	UE_LOG(LogConsoleVariablesEditor, Log, TEXT("Cmd: %s"), *FullCommand);
+}
+
 /** Get a reference to the cached console object. May return nullptr if unregistered. */
 IConsoleObject* FConsoleVariablesEditorCommandInfo::GetConsoleObjectPtr()
 {
@@ -178,6 +198,18 @@ FText FConsoleVariablesEditorCommandInfo::ConvertConsoleVariableSetByFlagToText(
 	return ReturnValue;
 }
 
+bool FConsoleVariablesEditorCommandInfo::GetCurrentValueAsString(FString& OutValueAsString)
+{
+	if (const IConsoleVariable* AsVariable = GetConsoleVariablePtr())
+	{
+		OutValueAsString = AsVariable->GetString();
+
+		return true;
+	}
+
+	return false;
+}
+
 bool FConsoleVariablesEditorCommandInfo::IsCurrentValueDifferentFromInputValue(const FString& InValueToCompare)
 {
 	if (const IConsoleVariable* AsVariable = GetConsoleVariablePtr())
diff --git a/Engine/Plugins/Editor/ConsoleVariablesEditor/Source/ConsoleVariablesEditor/Private/ConsoleVariablesEditorCommandInfo.h b/Engine/Plugins/Editor/ConsoleVariablesEditor/Source/ConsoleVariablesEditor/Private/ConsoleVariablesEditorCommandInfo.h
--- a/Engine/Plugins/Editor/ConsoleVariablesEditor/Source/ConsoleVariablesEditor/Private/ConsoleVariablesEditorCommandInfo.h
+++ b/Engine/Plugins/Editor/ConsoleVariablesEditor/Source/ConsoleVariablesEditor/Private/ConsoleVariablesEditorCommandInfo.h
@@ -2,13 +2,14 @@
 
 #pragma once
  
-#include "CoreMinimal.h"
 #include "ConsoleVariablesEditorModule.h"
 
 DECLARE_MULTICAST_DELEGATE_OneParam(FOnDetectConsoleObjectUnregistered, FString)
 
 #define LOCTEXT_NAMESPACE "ConsoleVariablesEditor"
 
+class IConsoleObject;
+
 struct FConsoleVariablesEditorCommandInfo
 {
 	enum class EConsoleObjectType
@@ -62,12 +63,21 @@ struct FConsoleVariablesEditorCommandInfo
 		bSetInCurrentSession = bNewSetting;
 	}
 	
-	/** Sets a variable to the specified value whilst maintaining its SetBy flag.
+	/**
+	  * Sets a variable to the specified value whilst maintaining its SetBy flag.
 	  * Non-variables will be executed through the console.
-	  * If bSetInSession is true, this CommandInfo's associated variable row will display "Session" in the UI.
+	  * @param bShouldTransactInConcert If true, send this value to multi-user.
+	  * @param bSetInSession If true, this CommandInfo's associated variable row will display "Session" in the UI.
+	  * @param bIsInteractiveChange If true, don't print the command or send it to multi-user/concert.
 	 */
 	void ExecuteCommand(
-		const FString& NewValueAsString, const bool bShouldTransactInConcert = true, const bool bSetInSession = true);
+		const FString& NewValueAsString, const bool bShouldTransactInConcert = true,
+		const bool bSetInSession = true, const bool bIsInteractiveChange = false);
+
+	/**
+	 * Prints the current value or state of the command, similar to the text generated in teh output log when entering a command manually.
+	 */
+	void PrintCommandOrVariable();
  
 	/** Get a reference to the cached console object. May return nullptr if unregistered. */
 	IConsoleObject* GetConsoleObjectPtr();
@@ -94,6 +104,12 @@ struct FConsoleVariablesEditorCommandInfo
 	FText GetSourceAsText();
 
 	static FText ConvertConsoleVariableSetByFlagToText(const EConsoleVariableFlags InFlag);
+
+	/**
+	 * If this console object is a variable, get its value as string via OutValueAsString.
+	 * If it's not a variable type (such as a command) then this function will return false.
+	 */
+	bool GetCurrentValueAsString(FString& OutValueAsString);
  
 	bool IsCurrentValueDifferentFromInputValue(const FString& InValueToCompare);
  
@@ -134,4 +150,5 @@ struct FConsoleVariablesEditorCommandInfo
 		{ EConsoleVariableFlags::ECVF_SetByConsole, LOCTEXT("Source_SetByConsole", "Console") }
 	};
 };
+
 #undef LOCTEXT_NAMESPACE
diff --git a/Engine/Plugins/Editor/ConsoleVariablesEditor/Source/ConsoleVariablesEditor/Private/Views/List/SConsoleVariablesEditorListValueInput.cpp b/Engine/Plugins/Editor/ConsoleVariablesEditor/Source/ConsoleVariablesEditor/Private/Views/List/SConsoleVariablesEditorListValueInput.cpp
--- a/Engine/Plugins/Editor/ConsoleVariablesEditor/Source/ConsoleVariablesEditor/Private/Views/List/SConsoleVariablesEditorListValueInput.cpp
+++ b/Engine/Plugins/Editor/ConsoleVariablesEditor/Source/ConsoleVariablesEditor/Private/Views/List/SConsoleVariablesEditorListValueInput.cpp
@@ -99,18 +99,12 @@ void SConsoleVariablesEditorListValueInput_Float::Construct(const FArguments& In
 
 			return FCString::Atof(*Item.Pin()->GetCachedValue());
 		})
-		.OnValueChanged_Lambda([this] (const float InValue)
+		.OnValueChanged(this, &SConsoleVariablesEditorListValueInput_Float::OnSliderAffected, false)
+		.OnValueCommitted_Lambda([this] (const float InValue, ETextCommit::Type CommitType)
 		{
-			check (Item.IsValid());
-					
-			const FString ValueAsString = FString::SanitizeFloat(InValue);
-					
-			if (const TSharedPtr<FConsoleVariablesEditorListRow> PinnedItem = Item.Pin();
-				!PinnedItem->GetCachedValue().Equals(ValueAsString))
-			{				
-				PinnedItem->GetCommandInfo().Pin()->ExecuteCommand(ValueAsString);
-					
-				PinnedItem->SetCachedValue(ValueAsString);
+			if (CommitType != ETextCommit::Default)
+			{
+				OnSliderAffected(InValue, true);
 			}
 		})
 		.IsEnabled(this, &SConsoleVariablesEditorListValueInput::IsRowChecked)
@@ -139,6 +133,18 @@ float SConsoleVariablesEditorListValueInput_Float::GetInputValue() const
 	return InputWidget->GetValue();
 }
 
+void SConsoleVariablesEditorListValueInput_Float::OnSliderAffected(const float InValue, const bool bPrintCommand)
+{
+	check (Item.IsValid());
+	const TSharedPtr<FConsoleVariablesEditorListRow> PinnedItem = Item.Pin();
+					
+	const FString ValueAsString = FString::SanitizeFloat(InValue);
+					
+	PinnedItem->GetCommandInfo().Pin()->ExecuteCommand(ValueAsString, true, true, !bPrintCommand);
+					
+	PinnedItem->SetCachedValue(ValueAsString);
+}
+
 void SConsoleVariablesEditorListValueInput_Int::Construct(const FArguments& InArgs,
                                                           const TWeakPtr<FConsoleVariablesEditorListRow> InRow,
                                                           const bool bIsShowFlag)
@@ -168,18 +174,12 @@ void SConsoleVariablesEditorListValueInput_Int::Construct(const FArguments& InAr
 
 			return FCString::Atoi(*Item.Pin()->GetCachedValue());
 		})
-		.OnValueChanged_Lambda([this] (const int32 InValue)
+		.OnValueChanged(this, &SConsoleVariablesEditorListValueInput_Int::OnSliderAffected, false)
+		.OnValueCommitted_Lambda([this] (const int32 InValue, ETextCommit::Type CommitType)
 		{
-			check (Item.IsValid());
-			
-			const FString ValueAsString = FString::FromInt(InValue);
-			
-			if (const TSharedPtr<FConsoleVariablesEditorListRow> PinnedItem = Item.Pin();
-				!PinnedItem->GetCachedValue().Equals(ValueAsString))
-			{				
-				PinnedItem->GetCommandInfo().Pin()->ExecuteCommand(ValueAsString);
-
-				PinnedItem->SetCachedValue(ValueAsString);
+			if (CommitType != ETextCommit::Default)
+			{
+				OnSliderAffected(InValue, true);
 			}
 		})
 		.IsEnabled(this, &SConsoleVariablesEditorListValueInput::IsRowChecked)
@@ -234,6 +234,30 @@ int32 SConsoleVariablesEditorListValueInput_Int::GetInputValue() const
 	return InputWidget->GetValue();
 }
 
+void SConsoleVariablesEditorListValueInput_Int::OnSliderAffected(const int32 InValue, const bool bPrintCommand)
+{
+	check (Item.IsValid());
+	const TSharedPtr<FConsoleVariablesEditorListRow> PinnedItem = Item.Pin();
+					
+	const FString ValueAsString = FString::FromInt(InValue);
+					
+	if (!PinnedItem->GetCachedValue().Equals(ValueAsString))
+	{				
+		PinnedItem->GetCommandInfo().Pin()->ExecuteCommand(ValueAsString, true, true, !bPrintCommand);
+					
+		PinnedItem->SetCachedValue(ValueAsString);
+
+		return;
+	}
+
+	// If the new value and cached value aren't different the command won't be executed.
+	// If we still want to print the command, do it here
+	if (bPrintCommand)
+	{
+		PinnedItem->GetCommandInfo().Pin()->PrintCommandOrVariable();
+	}
+}
+
 void SConsoleVariablesEditorListValueInput_String::Construct(const FArguments& InArgs,
                                                              const TWeakPtr<FConsoleVariablesEditorListRow> InRow)
 {
diff --git a/Engine/Plugins/Editor/ConsoleVariablesEditor/Source/ConsoleVariablesEditor/Private/Views/List/SConsoleVariablesEditorListValueInput.h b/Engine/Plugins/Editor/ConsoleVariablesEditor/Source/ConsoleVariablesEditor/Private/Views/List/SConsoleVariablesEditorListValueInput.h
--- a/Engine/Plugins/Editor/ConsoleVariablesEditor/Source/ConsoleVariablesEditor/Private/Views/List/SConsoleVariablesEditorListValueInput.h
+++ b/Engine/Plugins/Editor/ConsoleVariablesEditor/Source/ConsoleVariablesEditor/Private/Views/List/SConsoleVariablesEditorListValueInput.h
@@ -47,6 +47,7 @@ public:
 	virtual FString GetInputValueAsString() override;
 
 	float GetInputValue() const;
+	void OnSliderAffected(const float InValue, const bool bPrintCommand);
 
 private:
 	
@@ -71,6 +72,7 @@ public:
 	virtual FString GetInputValueAsString() override;
 
 	int32 GetInputValue() const;
+	void OnSliderAffected(const int32 InValue, const bool bPrintCommand);
 
 private:
 	
