Use proper ThirdParty include path for LibOVRPlatform header

#rnx

diff --git a/Engine/Plugins/Editor/ConsoleVariablesEditor/Source/ConsoleVariablesEditor/ConsoleVariablesEditor.build.cs b/Engine/Plugins/Editor/ConsoleVariablesEditor/Source/ConsoleVariablesEditor/ConsoleVariablesEditor.build.cs
--- a/Engine/Plugins/Editor/ConsoleVariablesEditor/Source/ConsoleVariablesEditor/ConsoleVariablesEditor.build.cs
+++ b/Engine/Plugins/Editor/ConsoleVariablesEditor/Source/ConsoleVariablesEditor/ConsoleVariablesEditor.build.cs
@@ -46,5 +46,12 @@ public class ConsoleVariablesEditor : ModuleRules
 				"WorkspaceMenuStructure"
 			}
 		);
+
+		PrivateIncludePathModuleNames.AddRange(
+			new string[]
+			{
+				"LibOVRPlatform"
+			}
+		);
 	}
 }
