Fix Linux include error

#rnx

diff --git a/Engine/Plugins/Editor/ConsoleVariablesEditor/Source/ConsoleVariablesEditor/ConsoleVariablesEditor.build.cs b/Engine/Plugins/Editor/ConsoleVariablesEditor/Source/ConsoleVariablesEditor/ConsoleVariablesEditor.build.cs
--- a/Engine/Plugins/Editor/ConsoleVariablesEditor/Source/ConsoleVariablesEditor/ConsoleVariablesEditor.build.cs
+++ b/Engine/Plugins/Editor/ConsoleVariablesEditor/Source/ConsoleVariablesEditor/ConsoleVariablesEditor.build.cs
@@ -46,12 +46,5 @@ public class ConsoleVariablesEditor : ModuleRules
 				"WorkspaceMenuStructure"
 			}
 		);
-
-		PrivateIncludePathModuleNames.AddRange(
-			new string[]
-			{
-				"LibOVRPlatform"
-			}
-		);
 	}
 }
