Invalidate our cache when rebuilding the list.

diff --git a/Engine/Plugins/Editor/ConsoleVariablesEditor/Source/ConsoleVariablesEditor/Private/ConsoleVariablesEditorModule.cpp b/Engine/Plugins/Editor/ConsoleVariablesEditor/Source/ConsoleVariablesEditor/Private/ConsoleVariablesEditorModule.cpp
--- a/Engine/Plugins/Editor/ConsoleVariablesEditor/Source/ConsoleVariablesEditor/Private/ConsoleVariablesEditorModule.cpp
+++ b/Engine/Plugins/Editor/ConsoleVariablesEditor/Source/ConsoleVariablesEditor/Private/ConsoleVariablesEditorModule.cpp
@@ -375,7 +375,7 @@ void FConsoleVariablesEditorModule::OnConsoleVariableChanged(IConsoleVariable* C
 			// If not yet tracked and we want to track variable changes from outside the dialogue,
 			// Check if the changed value differs from the startup value before tracking it
 			if (Settings->bAddAllChangedConsoleVariablesToCurrentPreset &&
-				!Settings->ChangedConsoleVariableSkipList.Contains(Key) && 
+				!Settings->ChangedConsoleVariableSkipList.Contains(Key) &&
 				PinnedCommand->IsCurrentValueDifferentFromInputValue(PinnedCommand->StartupValueAsString))
 			{
 				if (MainPanel.IsValid())
@@ -403,6 +403,11 @@ void FConsoleVariablesEditorModule::OnConsoleVariableChanged(IConsoleVariable* C
 		// If the variable is already tracked or was just added to tracking, run the following code
 		if (bIsVariableCurrentlyTracked)
 		{
+			if (MainPanel.IsValid())
+			{
+				MainPanel->UpdateCachedValue(Key, ChangedVariable->GetString());
+			}
+
 			/**
 			 * Here we check the map of recently received variables from other nodes
 			 * If the command is in the map and the value is similar, we won't send the value to other nodes
diff --git a/Engine/Plugins/Editor/ConsoleVariablesEditor/Source/ConsoleVariablesEditor/Private/Views/List/ConsoleVariablesEditorList.cpp b/Engine/Plugins/Editor/ConsoleVariablesEditor/Source/ConsoleVariablesEditor/Private/Views/List/ConsoleVariablesEditorList.cpp
--- a/Engine/Plugins/Editor/ConsoleVariablesEditor/Source/ConsoleVariablesEditor/Private/Views/List/ConsoleVariablesEditorList.cpp
+++ b/Engine/Plugins/Editor/ConsoleVariablesEditor/Source/ConsoleVariablesEditor/Private/Views/List/ConsoleVariablesEditorList.cpp
@@ -47,3 +47,11 @@ void FConsoleVariablesEditorList::UpdatePresetValuesForSave(TObjectPtr<UConsoleV
 		ListWidget->UpdatePresetValuesForSave(InAsset);
 	}
 }
+void FConsoleVariablesEditorList::UpdateCachedValue(const FString& InCommand, const FString& InValue)
+{
+	if (ListWidget.IsValid())
+	{
+		ListWidget->UpdateCachedValue(InCommand, InValue);
+	}
+}
+
diff --git a/Engine/Plugins/Editor/ConsoleVariablesEditor/Source/ConsoleVariablesEditor/Private/Views/List/ConsoleVariablesEditorList.h b/Engine/Plugins/Editor/ConsoleVariablesEditor/Source/ConsoleVariablesEditor/Private/Views/List/ConsoleVariablesEditorList.h
--- a/Engine/Plugins/Editor/ConsoleVariablesEditor/Source/ConsoleVariablesEditor/Private/Views/List/ConsoleVariablesEditorList.h
+++ b/Engine/Plugins/Editor/ConsoleVariablesEditor/Source/ConsoleVariablesEditor/Private/Views/List/ConsoleVariablesEditorList.h
@@ -55,6 +55,9 @@ public:
 	/** Updates the saved values in a UConsoleVariablesAsset so that the command/value map can be saved to disk */
 	void UpdatePresetValuesForSave(TObjectPtr<UConsoleVariablesAsset> InAsset) const;
 
+	/** Update cached value */
+	void UpdateCachedValue(const FString& Command, const FString& Value);
+
 private:
 
 	EConsoleVariablesEditorListMode CurrentListMode = EConsoleVariablesEditorListMode::Preset;
diff --git a/Engine/Plugins/Editor/ConsoleVariablesEditor/Source/ConsoleVariablesEditor/Private/Views/List/SConsoleVariablesEditorList.cpp b/Engine/Plugins/Editor/ConsoleVariablesEditor/Source/ConsoleVariablesEditor/Private/Views/List/SConsoleVariablesEditorList.cpp
--- a/Engine/Plugins/Editor/ConsoleVariablesEditor/Source/ConsoleVariablesEditor/Private/Views/List/SConsoleVariablesEditorList.cpp
+++ b/Engine/Plugins/Editor/ConsoleVariablesEditor/Source/ConsoleVariablesEditor/Private/Views/List/SConsoleVariablesEditorList.cpp
@@ -518,6 +518,22 @@ void SConsoleVariablesEditorList::UpdatePresetValuesForSave(const TObjectPtr<UCo
 	InAsset->ReplaceSavedCommands(NewSavedCommands);
 }
 
+void SConsoleVariablesEditorList::UpdateCachedValue(const FString& InCommand, const FString& InValue)
+{
+	const TArray<FConsoleVariablesEditorListRowPtr>& Items =
+		ListModelPtr.IsValid() && ListModelPtr.Pin()->GetListMode() ==
+			FConsoleVariablesEditorList::EConsoleVariablesEditorListMode::Preset ? TreeViewRootObjects : LastPresetObjects;
+	const FConsoleVariablesEditorListRowPtr* Item = Algo::FindByPredicate(Items, [InCommand](const FConsoleVariablesEditorListRowPtr& Item)
+	{
+		TSharedPtr<FConsoleVariablesEditorCommandInfo> CommandInfo = Item->GetCommandInfo().Pin();
+		return CommandInfo ? CommandInfo->Command == InCommand : false;
+	});
+	if (Item && (*Item)->IsRowChecked())
+	{
+		(*Item)->SetCachedValue(InValue);
+	}
+}
+
 FString SConsoleVariablesEditorList::GetSearchStringFromSearchInputField() const
 {
 	return ensureAlwaysMsgf(ListSearchBoxPtr.IsValid(),
diff --git a/Engine/Plugins/Editor/ConsoleVariablesEditor/Source/ConsoleVariablesEditor/Private/Views/List/SConsoleVariablesEditorList.h b/Engine/Plugins/Editor/ConsoleVariablesEditor/Source/ConsoleVariablesEditor/Private/Views/List/SConsoleVariablesEditorList.h
--- a/Engine/Plugins/Editor/ConsoleVariablesEditor/Source/ConsoleVariablesEditor/Private/Views/List/SConsoleVariablesEditorList.h
+++ b/Engine/Plugins/Editor/ConsoleVariablesEditor/Source/ConsoleVariablesEditor/Private/Views/List/SConsoleVariablesEditorList.h
@@ -72,6 +72,9 @@ public:
 	/** Updates the saved values in a UConsoleVariablesAsset so that the command/value map can be saved to disk */
 	void UpdatePresetValuesForSave(const TObjectPtr<UConsoleVariablesAsset> InAsset) const;
 
+	/** Update the cached value stored for the given command. */
+	void UpdateCachedValue(const FString& InCommand, const FString& InValue);
+
 	FString GetSearchStringFromSearchInputField() const;
 	void SetSearchStringInSearchInputField(const FString InSearchString) const;
 	void ExecuteListViewSearchOnAllRows(const FString& SearchString, const bool bShouldRefreshAfterward = true);
diff --git a/Engine/Plugins/Editor/ConsoleVariablesEditor/Source/ConsoleVariablesEditor/Private/Views/MainPanel/ConsoleVariablesEditorMainPanel.cpp b/Engine/Plugins/Editor/ConsoleVariablesEditor/Source/ConsoleVariablesEditor/Private/Views/MainPanel/ConsoleVariablesEditorMainPanel.cpp
--- a/Engine/Plugins/Editor/ConsoleVariablesEditor/Source/ConsoleVariablesEditor/Private/Views/MainPanel/ConsoleVariablesEditorMainPanel.cpp
+++ b/Engine/Plugins/Editor/ConsoleVariablesEditor/Source/ConsoleVariablesEditor/Private/Views/MainPanel/ConsoleVariablesEditorMainPanel.cpp
@@ -88,6 +88,14 @@ void FConsoleVariablesEditorMainPanel::AddConsoleObjectToCurrentPreset(
 	}
 }
 
+void FConsoleVariablesEditorMainPanel::UpdateCachedValue(const FString& InConsoleVariable, const FString& InValue)
+{
+	if (TSharedPtr<FConsoleVariablesEditorList> TheEditorList = GetEditorList().Pin())
+	{
+		TheEditorList->UpdateCachedValue(InConsoleVariable, InValue);
+	}
+}
+
 FReply FConsoleVariablesEditorMainPanel::ValidateConsoleInputAndAddToCurrentPreset(const FText& CommittedText) const
 {
 	// Clear Search
diff --git a/Engine/Plugins/Editor/ConsoleVariablesEditor/Source/ConsoleVariablesEditor/Private/Views/MainPanel/ConsoleVariablesEditorMainPanel.h b/Engine/Plugins/Editor/ConsoleVariablesEditor/Source/ConsoleVariablesEditor/Private/Views/MainPanel/ConsoleVariablesEditorMainPanel.h
--- a/Engine/Plugins/Editor/ConsoleVariablesEditor/Source/ConsoleVariablesEditor/Private/Views/MainPanel/ConsoleVariablesEditorMainPanel.h
+++ b/Engine/Plugins/Editor/ConsoleVariablesEditor/Source/ConsoleVariablesEditor/Private/Views/MainPanel/ConsoleVariablesEditorMainPanel.h
@@ -27,6 +27,8 @@ public:
 		const FString InConsoleCommand, const FString InValue,
 		const bool bScrollToNewRow = false) const;
 
+	void UpdateCachedValue(const FString& InConsoleVar, const FString& Value);
+
 	FReply ValidateConsoleInputAndAddToCurrentPreset(const FText& CommittedText) const;
 
 	/*
