Fix lifetime issues, all of these are undefined behaviour

diff --git a/Engine/Plugins/Editor/ConsoleVariablesEditor/Source/ConsoleVariablesEditor/Private/Views/List/SConsoleVariablesEditorList.h b/Engine/Plugins/Editor/ConsoleVariablesEditor/Source/ConsoleVariablesEditor/Private/Views/List/SConsoleVariablesEditorList.h
--- a/Engine/Plugins/Editor/ConsoleVariablesEditor/Source/ConsoleVariablesEditor/Private/Views/List/SConsoleVariablesEditorList.h
+++ b/Engine/Plugins/Editor/ConsoleVariablesEditor/Source/ConsoleVariablesEditor/Private/Views/List/SConsoleVariablesEditorList.h
@@ -187,33 +187,28 @@ private:
 	FName ActiveSortingColumnName = NAME_None;
 	EColumnSortMode::Type ActiveSortingType = EColumnSortMode::None;
 
-	TFunctionRef<bool(const FConsoleVariablesEditorListRowPtr&, const FConsoleVariablesEditorListRowPtr&)> SortByOrderAscending =
-		[](const FConsoleVariablesEditorListRowPtr& A, const FConsoleVariablesEditorListRowPtr& B)
+	TFunction<bool(const FConsoleVariablesEditorListRowPtr&, const FConsoleVariablesEditorListRowPtr&)> SortByOrderAscending = [](const FConsoleVariablesEditorListRowPtr& A, const FConsoleVariablesEditorListRowPtr& B)
 		{
 			return A->GetSortOrder() < B->GetSortOrder();
 		};
 
-	TFunctionRef<bool(const FConsoleVariablesEditorListRowPtr&, const FConsoleVariablesEditorListRowPtr&)> SortBySourceAscending =
-		[](const FConsoleVariablesEditorListRowPtr& A, const FConsoleVariablesEditorListRowPtr& B)
+	TFunction<bool(const FConsoleVariablesEditorListRowPtr&, const FConsoleVariablesEditorListRowPtr&)> SortBySourceAscending = [](const FConsoleVariablesEditorListRowPtr& A, const FConsoleVariablesEditorListRowPtr& B)
 		{
 			return A->GetCommandInfo().Pin()->GetSourceAsText().ToString() < B->GetCommandInfo().Pin()->GetSourceAsText().ToString();
 		};
 
-	TFunctionRef<bool(const FConsoleVariablesEditorListRowPtr&, const FConsoleVariablesEditorListRowPtr&)> SortBySourceDescending =
-		[](const FConsoleVariablesEditorListRowPtr& A, const FConsoleVariablesEditorListRowPtr& B)
+	TFunction<bool(const FConsoleVariablesEditorListRowPtr&, const FConsoleVariablesEditorListRowPtr&)> SortBySourceDescending = [](const FConsoleVariablesEditorListRowPtr& A, const FConsoleVariablesEditorListRowPtr& B)
 		{
 			return B->GetCommandInfo().Pin()->GetSourceAsText().ToString() < A->GetCommandInfo().Pin()->GetSourceAsText().ToString();
 		};
-	
-	TFunctionRef<bool(const FConsoleVariablesEditorListRowPtr&, const FConsoleVariablesEditorListRowPtr&)> SortByVariableNameAscending =
-		[](const FConsoleVariablesEditorListRowPtr& A, const FConsoleVariablesEditorListRowPtr& B)
-	{
-		return A->GetCommandInfo().Pin()->Command < B->GetCommandInfo().Pin()->Command;
-	};
 
-	TFunctionRef<bool(const FConsoleVariablesEditorListRowPtr&, const FConsoleVariablesEditorListRowPtr&)> SortByVariableNameDescending =
-		[](const FConsoleVariablesEditorListRowPtr& A, const FConsoleVariablesEditorListRowPtr& B)
-	{
-		return B->GetCommandInfo().Pin()->Command < A->GetCommandInfo().Pin()->Command;
-	};
+	TFunction<bool(const FConsoleVariablesEditorListRowPtr&, const FConsoleVariablesEditorListRowPtr&)> SortByVariableNameAscending = [](const FConsoleVariablesEditorListRowPtr& A, const FConsoleVariablesEditorListRowPtr& B)
+		{
+			return A->GetCommandInfo().Pin()->Command < B->GetCommandInfo().Pin()->Command;
+		};
+
+	TFunction<bool(const FConsoleVariablesEditorListRowPtr&, const FConsoleVariablesEditorListRowPtr&)> SortByVariableNameDescending = [](const FConsoleVariablesEditorListRowPtr& A, const FConsoleVariablesEditorListRowPtr& B)
+		{
+			return B->GetCommandInfo().Pin()->Command < A->GetCommandInfo().Pin()->Command;
+		};
 };
diff --git a/Engine/Plugins/Editor/ObjectMixer/ObjectMixer/Source/ObjectMixer/Public/Views/List/SObjectMixerEditorList.h b/Engine/Plugins/Editor/ObjectMixer/ObjectMixer/Source/ObjectMixer/Public/Views/List/SObjectMixerEditorList.h
--- a/Engine/Plugins/Editor/ObjectMixer/ObjectMixer/Source/ObjectMixer/Public/Views/List/SObjectMixerEditorList.h
+++ b/Engine/Plugins/Editor/ObjectMixer/ObjectMixer/Source/ObjectMixer/Public/Views/List/SObjectMixerEditorList.h
@@ -293,8 +293,7 @@ protected:
 	FName ActiveSortingColumnName = NAME_None;
 	EColumnSortMode::Type ActiveSortingType = EColumnSortMode::None;
 
-	TFunctionRef<bool(const FObjectMixerEditorListRowPtr&, const FObjectMixerEditorListRowPtr&)> SortByOrderAscending =
-		[](const FObjectMixerEditorListRowPtr& A, const FObjectMixerEditorListRowPtr& B)
+	TFunction<bool(const FObjectMixerEditorListRowPtr&, const FObjectMixerEditorListRowPtr&)> SortByOrderAscending = [](const FObjectMixerEditorListRowPtr& A, const FObjectMixerEditorListRowPtr& B)
 		{
 			return A->GetSortOrder() < B->GetSortOrder();
 		};
diff --git a/Engine/Plugins/Enterprise/DataprepEditor/Source/DataprepEditor/Private/Widgets/DataprepGraph/SDataprepGraphTrackNode.cpp b/Engine/Plugins/Enterprise/DataprepEditor/Source/DataprepEditor/Private/Widgets/DataprepGraph/SDataprepGraphTrackNode.cpp
--- a/Engine/Plugins/Enterprise/DataprepEditor/Source/DataprepEditor/Private/Widgets/DataprepGraph/SDataprepGraphTrackNode.cpp
+++ b/Engine/Plugins/Enterprise/DataprepEditor/Source/DataprepEditor/Private/Widgets/DataprepGraph/SDataprepGraphTrackNode.cpp
@@ -36,7 +36,7 @@
 
 FVector2D SDataprepGraphTrackNode::TrackAnchor( 15.f, 40.f );
 
-typedef TFunctionRef<float(const FVector2D&, float)> DragCallback;
+typedef TFunction<float(const FVector2D&, float)> DragCallback;
 
 class SDataprepEmptyActionNode : public SVerticalBox
 {
diff --git a/Engine/Plugins/Online/EOSShared/Source/EOSShared/Private/EOSShared.cpp b/Engine/Plugins/Online/EOSShared/Source/EOSShared/Private/EOSShared.cpp
--- a/Engine/Plugins/Online/EOSShared/Source/EOSShared/Private/EOSShared.cpp
+++ b/Engine/Plugins/Online/EOSShared/Source/EOSShared/Private/EOSShared.cpp
@@ -215,7 +215,7 @@ bool LexFromString(EOS_EAuthScopeFlags& OutEnum, const FStringView InString)
 
 	using namespace UE::String;
 	const EParseTokensOptions ParseOptions = EParseTokensOptions::SkipEmpty | EParseTokensOptions::Trim;
-	TFunctionRef<void(FStringView)> ParseFunc = [&OutEnum, &bParsedOk](FStringView Token)
+	auto ParseFunc = [&OutEnum, &bParsedOk](FStringView Token)
 	{
 		if (Token == TEXT("BasicProfile"))
 		{
@@ -247,7 +247,8 @@ bool LexFromString(EOS_EAuthScopeFlags& OutEnum, const FStringView InString)
 			bParsedOk = false;
 		}
 	};
-	ParseTokens(InString, TCHAR('|'), ParseFunc, ParseOptions);
+
+	ParseTokens(InString, TCHAR('|'), (TFunctionRef<void(FStringView)>)ParseFunc, ParseOptions);
 
 	return bParsedOk;
 }
diff --git a/Engine/Plugins/Runtime/GeometryProcessing/Source/DynamicMesh/Public/Operations/RemoveOccludedTriangles.h b/Engine/Plugins/Runtime/GeometryProcessing/Source/DynamicMesh/Public/Operations/RemoveOccludedTriangles.h
--- a/Engine/Plugins/Runtime/GeometryProcessing/Source/DynamicMesh/Public/Operations/RemoveOccludedTriangles.h
+++ b/Engine/Plugins/Runtime/GeometryProcessing/Source/DynamicMesh/Public/Operations/RemoveOccludedTriangles.h
@@ -213,7 +213,7 @@ public:
 			}
 		};
 
-		typedef TFunctionRef<bool(const TArrayView<TMeshAABBTree3<OccluderTriangleMeshType>*> Spatials, const TArrayView<TFastWindingTree<OccluderTriangleMeshType>*> FastWindingTrees,
+		typedef TFunction<bool(const TArrayView<TMeshAABBTree3<OccluderTriangleMeshType>*> Spatials, const TArrayView<TFastWindingTree<OccluderTriangleMeshType>*> FastWindingTrees,
 			const FVector3d& Pt, const FVector3d& OriginalPt, const FTransformSRT3d& OriginalXF, const TArray<FTransformWithInv>& SpatialTransforms)> FIsOccludedFn;
 
 		// test if point is occluded by the combination of all spatials, according to winding number
@@ -625,4 +625,4 @@ protected:
 
 
 } // end namespace UE::Geometry
-} // end namespace UE
\ No newline at end of file
+} // end namespace UE
diff --git a/Engine/Source/Developer/Virtualization/Private/VirtualizationManager.cpp b/Engine/Source/Developer/Virtualization/Private/VirtualizationManager.cpp
--- a/Engine/Source/Developer/Virtualization/Private/VirtualizationManager.cpp
+++ b/Engine/Source/Developer/Virtualization/Private/VirtualizationManager.cpp
@@ -1128,7 +1128,7 @@ void FVirtualizationManager::ApplySettingsFromConfigFiles(const FConfigFile& Con
 	// Check for any legacy settings and print them out (easier to do this in one block rather than one and time)
 	{
 		// Entries that are allows to be in [Core.ContentVirtualization
-		TArrayView<const TCHAR* const> AllowedEntries = { TEXT("SystemName") , TEXT("LazyInit") };
+		static const TArray<FString> AllowedEntries = { TEXT("SystemName") , TEXT("LazyInit") };
 		
 		TArray<FString> LegacyEntries;	
 		if (const FConfigSection* LegacySection = ConfigFile.Find(LegacyConfigSection))
@@ -2083,4 +2083,4 @@ void FVirtualizationManager::OnAnalyticsEvent(TArray<FAnalyticsEventAttribute>&
 } // namespace UE::Virtualization
 
 #undef UE_INLINE_ALLOCATION_COUNT
-#undef LOCTEXT_NAMESPACE
\ No newline at end of file
+#undef LOCTEXT_NAMESPACE
diff --git a/Engine/Source/Runtime/CoreUObject/Private/Misc/ReferenceDebugging.cpp b/Engine/Source/Runtime/CoreUObject/Private/Misc/ReferenceDebugging.cpp
--- a/Engine/Source/Runtime/CoreUObject/Private/Misc/ReferenceDebugging.cpp
+++ b/Engine/Source/Runtime/CoreUObject/Private/Misc/ReferenceDebugging.cpp
@@ -7,7 +7,7 @@
 void FindPackageReferencesInObject(const UObject* RootObject, const TArray<FName>& PackagesReferenced, TFunctionRef<void(FName /*PackagePath*/, const FString& /*PropertyPath*/)> Callback)
 {
 	TArray<FString> PathStack;
-	TFunctionRef<void(const UObject*)> InnerLoop = [&InnerLoop, &PackagesReferenced, &PathStack, Callback](const UObject* RootObject)
+	TFunction<void(const UObject*)> InnerLoop = [&InnerLoop, &PackagesReferenced, &PathStack, Callback](const UObject* RootObject)
 	{
 		for (TPropertyValueIterator<FObjectProperty> It(RootObject->GetClass(), RootObject, EPropertyValueIteratorFlags::FullRecursion, EFieldIteratorFlags::ExcludeDeprecated); It; ++It)
 		{
@@ -57,4 +57,4 @@ void FindPackageReferencesInObject(const UObject* RootObject, const TArray<FName
 	};
 
 	InnerLoop(RootObject);
-}
\ No newline at end of file
+}
diff --git a/Engine/Source/Runtime/Engine/Classes/Engine/Level.h b/Engine/Source/Runtime/Engine/Classes/Engine/Level.h
--- a/Engine/Source/Runtime/Engine/Classes/Engine/Level.h
+++ b/Engine/Source/Runtime/Engine/Classes/Engine/Level.h
@@ -1309,7 +1309,8 @@ public:
 	/** 
 	 * Call on a level that was loaded from disk instead of PIE-duplicating, to fixup actor references
 	 */
-	ENGINE_API void FixupForPIE(int32 PIEInstanceID, TFunctionRef<void(int32, FSoftObjectPath&)> CustomFixupFunction = [](int32, FSoftObjectPath&) {});
+	ENGINE_API void FixupForPIE(int32 PIEInstanceID);
+	ENGINE_API void FixupForPIE(int32 PIEInstanceID, TFunctionRef<void(int32, FSoftObjectPath&)> CustomFixupFunction);
 
 	/**
 	 * Returns true if the level contains static meshes that have not finished compiling yet.
diff --git a/Engine/Source/Runtime/Engine/Private/Animation/AnimClassData.cpp b/Engine/Source/Runtime/Engine/Private/Animation/AnimClassData.cpp
--- a/Engine/Source/Runtime/Engine/Private/Animation/AnimClassData.cpp
+++ b/Engine/Source/Runtime/Engine/Private/Animation/AnimClassData.cpp
@@ -55,6 +55,5 @@ PRAGMA_ENABLE_DEPRECATION_WARNINGS
 
 TArrayView<const FAnimNodeData> UAnimClassData::GetNodeData() const
 {
-	TArray<FAnimNodeData> Dummy;
-	return Dummy;
+	return {};
 }
diff --git a/Engine/Source/Runtime/Engine/Private/Level.cpp b/Engine/Source/Runtime/Engine/Private/Level.cpp
--- a/Engine/Source/Runtime/Engine/Private/Level.cpp
+++ b/Engine/Source/Runtime/Engine/Private/Level.cpp
@@ -4199,6 +4199,11 @@ void ULevel::FixupForPIE(int32 InPIEInstanceID, TFunctionRef<void(int32, FSoftOb
 	Serialize(FixupSerializer);
 }
 
+void ULevel::FixupForPIE(int32 InPIEInstanceID)
+{
+	FixupForPIE(InPIEInstanceID, [](int32, FSoftObjectPath&) {});
+}
+
 #endif	//WITH_EDITOR
 
 bool ULevel::ResolveSubobject(const TCHAR* SubObjectPath, UObject*& OutObject, bool bLoadIfExists)
diff --git a/Engine/Source/Runtime/Engine/Private/PieFixupSerializer.cpp b/Engine/Source/Runtime/Engine/Private/PieFixupSerializer.cpp
--- a/Engine/Source/Runtime/Engine/Private/PieFixupSerializer.cpp
+++ b/Engine/Source/Runtime/Engine/Private/PieFixupSerializer.cpp
@@ -7,8 +7,13 @@
 
 class FMulticastDelegateProperty;
 
+namespace
+{
+	void DefaultSoftObjectPathFixupFunction(int32, FSoftObjectPath&) {}
+}
+
 FPIEFixupSerializer::FPIEFixupSerializer(UObject* InRoot, int32 InPIEInstanceID)
-	: SoftObjectPathFixupFunction([](int32, FSoftObjectPath&) {})
+	: SoftObjectPathFixupFunction(DefaultSoftObjectPathFixupFunction)
 	, Root(InRoot)
 	, PIEInstanceID(InPIEInstanceID)
 {
@@ -69,4 +74,4 @@ FArchive& FPIEFixupSerializer::operator<<(FSoftObjectPtr& Value)
 	// Forward the serialization to the FSoftObjectPath overload so it can be fixed up
 	*this << Value.GetUniqueID();
 	return *this;
-}
\ No newline at end of file
+}
diff --git a/Engine/Source/Runtime/GeometryCore/Private/CompGeom/FitOrientedBox3.cpp b/Engine/Source/Runtime/GeometryCore/Private/CompGeom/FitOrientedBox3.cpp
--- a/Engine/Source/Runtime/GeometryCore/Private/CompGeom/FitOrientedBox3.cpp
+++ b/Engine/Source/Runtime/GeometryCore/Private/CompGeom/FitOrientedBox3.cpp
@@ -126,7 +126,7 @@ TOrientedBox3<RealType> OptimizeOrientedBox3Points(const TOrientedBox3<RealType>
 			}
 		}
 		RealType Depth = BestBox.Extents[Axis] * 2;
-		using TFitFn = TFunctionRef<RealType(RealType, RealType)>;
+		using TFitFn = TFunction<RealType(RealType, RealType)>;
 		TFitFn Score2 = FitMethod == EBox3FitCriteria::Volume ?
 			TFitFn([](RealType Width, RealType Height) -> RealType
 				{
@@ -371,7 +371,7 @@ TOrientedBox3<RealType> FitOrientedBox3Points(int32 NumPts, TFunctionRef<TVector
 			ProjBoundary.Emplace(BasisX.Dot(Pt), BasisY.Dot(Pt));
 		}
 		
-		using TFitFn = TFunctionRef<RealType(RealType, RealType)>;
+		using TFitFn = TFunction<RealType(RealType, RealType)>;
 		TFitFn Score2 = FitMethod == EBox3FitCriteria::Volume ? 
 			TFitFn([](RealType Width, RealType Height) -> RealType
 			{
@@ -435,4 +435,4 @@ template TOrientedBox3<double> GEOMETRYCORE_API OptimizeOrientedBox3Points(const
 
 
 }
-}
\ No newline at end of file
+}
