ConsoleVariablesEditor: Make FConsoleVariablesEditorCommandInfo and CommandInfo accessors in FConsoleVariablesEditorModule public. This is in preparation
for integrating CommandInfo into Movie Render Queue.

diff --git a/Engine/Plugins/Editor/ConsoleVariablesEditor/Source/ConsoleVariablesEditor/Private/ConsoleVariablesEditorCommandInfo.cpp b/Engine/Plugins/Editor/ConsoleVariablesEditor/Source/ConsoleVariablesEditor/Private/ConsoleVariablesEditorCommandInfo.cpp
--- a/Engine/Plugins/Editor/ConsoleVariablesEditor/Source/ConsoleVariablesEditor/Private/ConsoleVariablesEditorCommandInfo.cpp
+++ b/Engine/Plugins/Editor/ConsoleVariablesEditor/Source/ConsoleVariablesEditor/Private/ConsoleVariablesEditorCommandInfo.cpp
@@ -3,12 +3,13 @@
 #include "ConsoleVariablesEditorCommandInfo.h"
 
 #include "ConsoleVariablesEditorLog.h"
+#include "ConsoleVariablesEditorModule.h"
+#include "MultiUser/ConsoleVariableSync.h"
 
 #include "Algo/Find.h"
 #include "Editor.h"
 #include "Editor/EditorEngine.h"
 #include "Engine/GameEngine.h"
-#include "HAL/IConsoleManager.h"
 
 #define LOCTEXT_NAMESPACE "ConsoleVariablesEditor"
 
diff --git a/Engine/Plugins/Editor/ConsoleVariablesEditor/Source/ConsoleVariablesEditor/Private/ConsoleVariablesEditorCommandInfo.h b/Engine/Plugins/Editor/ConsoleVariablesEditor/Source/ConsoleVariablesEditor/Private/ConsoleVariablesEditorCommandInfo.h
deleted file mode 100644
diff --git a/Engine/Plugins/Editor/ConsoleVariablesEditor/Source/ConsoleVariablesEditor/Private/ConsoleVariablesEditorModule.cpp b/Engine/Plugins/Editor/ConsoleVariablesEditor/Source/ConsoleVariablesEditor/Private/ConsoleVariablesEditorModule.cpp
--- a/Engine/Plugins/Editor/ConsoleVariablesEditor/Source/ConsoleVariablesEditor/Private/ConsoleVariablesEditorModule.cpp
+++ b/Engine/Plugins/Editor/ConsoleVariablesEditor/Source/ConsoleVariablesEditor/Private/ConsoleVariablesEditorModule.cpp
@@ -3,6 +3,7 @@
 #include "ConsoleVariablesEditorModule.h"
 
 #include "AssetTypeActions/AssetTypeActions_ConsoleVariables.h"
+#include "ConsoleVariablesAsset.h"
 #include "ConsoleVariablesEditorCommandInfo.h"
 #include "ConsoleVariablesEditorLog.h"
 #include "ConsoleVariablesEditorProjectSettings.h"
diff --git a/Engine/Plugins/Editor/ConsoleVariablesEditor/Source/ConsoleVariablesEditor/Private/ConsoleVariablesEditorModule.h b/Engine/Plugins/Editor/ConsoleVariablesEditor/Source/ConsoleVariablesEditor/Private/ConsoleVariablesEditorModule.h
deleted file mode 100644
diff --git a/Engine/Plugins/Editor/ConsoleVariablesEditor/Source/ConsoleVariablesEditor/Public/ConsoleVariablesEditorCommandInfo.h b/Engine/Plugins/Editor/ConsoleVariablesEditor/Source/ConsoleVariablesEditor/Public/ConsoleVariablesEditorCommandInfo.h
new file mode 100644
--- /dev/null
+++ b/Engine/Plugins/Editor/ConsoleVariablesEditor/Source/ConsoleVariablesEditor/Public/ConsoleVariablesEditorCommandInfo.h
@@ -0,0 +1,157 @@
+// Copyright Epic Games, Inc. All Rights Reserved.
+
+#pragma once
+
+#include "Containers/UnrealString.h"
+#include "HAL/IConsoleManager.h"
+#include "Internationalization/Internationalization.h"
+#include "Internationalization/Text.h"
+#include "Misc/DateTime.h"
+#include "UObject/ObjectMacros.h"
+
+DECLARE_MULTICAST_DELEGATE_OneParam(FOnDetectConsoleObjectUnregistered, FString)
+
+#define LOCTEXT_NAMESPACE "ConsoleVariablesEditor"
+
+struct CONSOLEVARIABLESEDITOR_API FConsoleVariablesEditorCommandInfo
+{
+	enum class EConsoleObjectType
+	{
+		// A console command that has no associated console object but is parsed externally, e.g. 'stat unit'
+		NullObject,
+		// A console command with an associated console object, like 'r.SetNearClipPlane'
+		Command,
+		// A console variable such as 'r.ScreenPercentage'
+		Variable
+	};
+	
+	struct FStaticConsoleVariableFlagInfo
+	{
+		EConsoleVariableFlags Flag;
+		FText DisplayText;
+	};
+	
+	FConsoleVariablesEditorCommandInfo(const FString& InCommand)
+	: Command(InCommand)
+	{
+		if (GetConsoleObjectPtr())
+		{
+			ObjectType = EConsoleObjectType::Command;
+			
+			if (const IConsoleVariable* AsVariable = GetConsoleVariablePtr())
+			{
+				ObjectType = EConsoleObjectType::Variable;
+				StartupValueAsString = AsVariable->GetString();
+				StartupSource = GetSource();
+			}
+		}
+	}
+	
+	~FConsoleVariablesEditorCommandInfo()
+	{
+		OnDetectConsoleObjectUnregistered.Remove(OnDetectConsoleObjectUnregisteredHandle);
+		if (IConsoleVariable* AsVariable = GetConsoleVariablePtr())
+		{
+			AsVariable->OnChangedDelegate().Remove(OnVariableChangedCallbackHandle);
+		}
+	}
+	
+	FORCEINLINE bool operator==(const FConsoleVariablesEditorCommandInfo& Comparator) const
+	{
+		return Command.Equals(Comparator.Command, ESearchCase::IgnoreCase);
+	}
+	
+	void SetIfChangedInCurrentPreset(const bool bNewSetting)
+	{
+		bSetInCurrentSession = bNewSetting;
+	}
+	
+	/**
+	  * Sets a variable to the specified value whilst maintaining its SetBy flag.
+	  * Non-variables will be executed through the console.
+	  * @param bShouldTransactInConcert If true, send this value to multi-user.
+	  * @param bSetInSession If true, this CommandInfo's associated variable row will display "Session" in the UI.
+	  * @param bIsInteractiveChange If true, don't print the command or send it to multi-user/concert.
+	 */
+	void ExecuteCommand(
+		const FString& NewValueAsString, const bool bShouldTransactInConcert = true,
+		const bool bSetInSession = true, const bool bIsInteractiveChange = false);
+
+	/**
+	 * Prints the current value or state of the command, similar to the text generated in teh output log when entering a command manually.
+	 */
+	void PrintCommandOrVariable();
+ 
+	/** Get a reference to the cached console object. May return nullptr if unregistered. */
+	IConsoleObject* GetConsoleObjectPtr();
+
+	/** Return the console object as a console variable object if applicable. May return nullptr if unregistered. */
+	IConsoleVariable* GetConsoleVariablePtr();
+
+	/**
+	 *Return the console object as a console command object if applicable.
+	 *Does not consider externally parsed console commands, as they have no associated objects.
+	 */
+	IConsoleCommand* GetConsoleCommandPtr();
+ 
+	static UWorld* GetCurrentWorld();
+
+	FString GetHelpText();
+ 
+	EConsoleVariableFlags GetSource();
+ 
+	void ClearSourceFlags();
+
+	void SetSourceFlag(const EConsoleVariableFlags InSource);
+ 
+	FText GetSourceAsText();
+
+	static FText ConvertConsoleVariableSetByFlagToText(const EConsoleVariableFlags InFlag);
+
+	/**
+	 * If this console object is a variable, get its value as string via OutValueAsString.
+	 * If it's not a variable type (such as a command) then this function will return false.
+	 */
+	bool GetCurrentValueAsString(FString& OutValueAsString);
+ 
+	bool IsCurrentValueDifferentFromInputValue(const FString& InValueToCompare);
+ 
+	/** The actual string key or name */
+	UPROPERTY()
+	FString Command;
+	EConsoleObjectType ObjectType = EConsoleObjectType::NullObject;
+	/** This object is periodically refreshed to mitigate the occurrence of stale pointers. */
+	IConsoleObject* ConsoleObjectPtr;
+	FDateTime TimeOfLastConsoleObjectRefresh;
+	double ConsoleObjectRefreshThreshold = 1.0;
+	
+	/** The value of this variable (if Variable object type) when the module started in this session after it may have been set by an ini file. */
+	FString StartupValueAsString;
+	
+	/** The source of this variable's (if Variable object type) last setting as recorded when the plugin was loaded. */
+	EConsoleVariableFlags StartupSource = ECVF_Default;
+	/** If the variable was last changed by the current preset */
+	bool bSetInCurrentSession = false;
+	/** When variables change, this callback is executed. */
+	FDelegateHandle OnVariableChangedCallbackHandle;
+	/** When commands are unregistered change, this callback is broadcasted. */
+	FOnDetectConsoleObjectUnregistered OnDetectConsoleObjectUnregistered;
+	FDelegateHandle OnDetectConsoleObjectUnregisteredHandle;
+	/** A mapping of SetBy console variable flags to information like the associated display text. */
+	static const inline TArray<FStaticConsoleVariableFlagInfo> SupportedFlags =
+	{
+		{ EConsoleVariableFlags::ECVF_SetByConstructor, LOCTEXT("Source_SetByConstructor", "Constructor") },
+		{ EConsoleVariableFlags::ECVF_SetByScalability, LOCTEXT("Source_SetByScalability", "Scalability") },
+		{ EConsoleVariableFlags::ECVF_SetByGameSetting, LOCTEXT("Source_SetByGameSetting", "Game Setting") },
+		{ EConsoleVariableFlags::ECVF_SetByProjectSetting, LOCTEXT("Source_SetByProjectSetting", "Project Setting") },
+		{ EConsoleVariableFlags::ECVF_SetBySystemSettingsIni, LOCTEXT("Source_SetBySystemSettingsIni", "System Settings ini") },
+		{ EConsoleVariableFlags::ECVF_SetByDeviceProfile, LOCTEXT("Source_SetByDeviceProfile", "Device Profile") },
+		{ EConsoleVariableFlags::ECVF_SetByGameOverride, LOCTEXT("Source_SetByGameOverride", "Game Override") },
+		{ EConsoleVariableFlags::ECVF_SetByConsoleVariablesIni, LOCTEXT("Source_SetByConsoleVariablesIni", "Console Variables ini") },
+		{ EConsoleVariableFlags::ECVF_SetByCommandline, LOCTEXT("Source_SetByCommandline", "Command line") },
+		{ EConsoleVariableFlags::ECVF_SetByCode, LOCTEXT("Source_SetByCode", "Code") },
+		{ EConsoleVariableFlags::ECVF_SetByConsole, LOCTEXT("Source_SetByConsole", "Console") }
+	};
+};
+
+#undef LOCTEXT_NAMESPACE
diff --git a/Engine/Plugins/Editor/ConsoleVariablesEditor/Source/ConsoleVariablesEditor/Public/ConsoleVariablesEditorModule.h b/Engine/Plugins/Editor/ConsoleVariablesEditor/Source/ConsoleVariablesEditor/Public/ConsoleVariablesEditorModule.h
new file mode 100644
--- /dev/null
+++ b/Engine/Plugins/Editor/ConsoleVariablesEditor/Source/ConsoleVariablesEditor/Public/ConsoleVariablesEditorModule.h
@@ -0,0 +1,170 @@
+// Copyright Epic Games, Inc. All Rights Reserved.
+
+#pragma once
+
+#include "CoreMinimal.h"
+#include "HAL/IConsoleManager.h"
+#include "Modules/ModuleInterface.h"
+#include "Widgets/Docking/SDockTab.h"
+
+struct FAssetData;
+struct FConsoleVariablesEditorCommandInfo;
+
+class FConsoleVariablesEditorMainPanel;
+class FConsoleVariablesEditorToolkit;
+class ISettingsSection;
+class UConsoleVariablesAsset;
+
+enum class ERemoteCVarChangeType : uint8;
+
+class FConsoleVariablesEditorModule : public IModuleInterface, public FGCObject
+{
+public:
+	
+	CONSOLEVARIABLESEDITOR_API static FConsoleVariablesEditorModule& Get();
+
+	//~ Begin IModuleInterface Interface
+	virtual void StartupModule() override;
+	virtual void ShutdownModule() override;
+	//~ End IModuleInterface Interface
+
+	void SavePreset() const;
+	void SaveSpecificPreset(const TObjectPtr<UConsoleVariablesAsset> Preset) const;
+	void SavePresetAs() const;
+	void SaveSpecificPresetAs(const TObjectPtr<UConsoleVariablesAsset> Preset) const;
+	void OpenConsoleVariablesDialogWithPreset(const TObjectPtr<UConsoleVariablesAsset> Preset) const;
+	void OpenConsoleVariablesDialogWithAssetSelected(const FAssetData& InAssetData) const;
+
+	/** Find all console variables and cache their startup values */
+	void QueryAndBeginTrackingConsoleVariables();
+
+	void AddConsoleObjectCommandInfoToMainReference(TSharedRef<FConsoleVariablesEditorCommandInfo> InCommandInfo)
+	{
+		ConsoleObjectsMainReference.Add(InCommandInfo);
+	}
+
+	/** Find a tracked console variable by the command string with optional case sensitivity. */
+	CONSOLEVARIABLESEDITOR_API TWeakPtr<FConsoleVariablesEditorCommandInfo> FindCommandInfoByName(
+		const FString& NameToSearch, ESearchCase::Type InSearchCase = ESearchCase::IgnoreCase);
+
+	/*Find all tracked console variables matching a specific search query with optional case sensitivity.
+	 *Individual members of InTokens will be considered "AnyOf" or "OR" searches. If SearchTerms contains any individual member it will match.
+	 *Members will be tested for a space character (" "). If a space is found, a subsearch will be run.
+	 *This subsearch will be an "AllOf" or "AND" type search in which all strings, separated by a space, must be found in the search terms.
+	 */
+	CONSOLEVARIABLESEDITOR_API TArray<TWeakPtr<FConsoleVariablesEditorCommandInfo>> FindCommandInfosMatchingTokens(
+		const TArray<FString>& InTokens, ESearchCase::Type InSearchCase = ESearchCase::IgnoreCase);
+
+	/**
+	 *Find a tracked console variable by its console object reference.
+	 *Note that some commands do not have an associated console object (such as 'stat unit') and will not be found with this method.
+	 *It's normally safer to use FindCommandInfoByName() instead.
+	 */
+	CONSOLEVARIABLESEDITOR_API TWeakPtr<FConsoleVariablesEditorCommandInfo> FindCommandInfoByConsoleObjectReference(
+		IConsoleObject* InConsoleObjectReference);
+	
+	[[nodiscard]] TObjectPtr<UConsoleVariablesAsset> GetPresetAsset() const;
+	[[nodiscard]] TObjectPtr<UConsoleVariablesAsset> GetGlobalSearchAsset() const;
+	
+	FReply ValidateConsoleInputAndAddToCurrentPreset(const FText& CommittedText) const;
+	/*
+	 * Regenerate the list items and refresh the list. Call when adding or removing variables.
+	 * @param bShouldCacheValues If true, the current list's current values will be cached and then restored when the list is rebuilt. Otherwise preset values will be used.
+	 */
+	void RebuildList(const FString InConsoleCommandToScrollTo = "", bool bShouldCacheValues = true) const;
+	
+	/**
+	 * Refresh filters and sorting.
+	 * Useful for when the list state has gone stale but the variable count has not changed.
+	 */
+	void RefreshList() const;
+	
+	void UpdatePresetValuesForSave(TObjectPtr<UConsoleVariablesAsset> InAsset);
+
+	/** Fills Global Search Asset's Saved Commands with variables matching the specified query. Returns false if no matches were found. */
+	bool PopulateGlobalSearchAssetWithVariablesMatchingTokens(const TArray<FString>& InTokens);
+
+	void SendMultiUserConsoleVariableChange(ERemoteCVarChangeType InChangeType, const FString& InVariableName, const FString& InValueAsString);
+	void OnRemoteCvarChanged(ERemoteCVarChangeType InChangeType, FString InName, FString InValue);
+
+	virtual void AddReferencedObjects( FReferenceCollector& Collector )  override;
+	virtual FString GetReferencerName() const override;
+
+private:
+	/** Scoped based struct to track inbound cvars set by Multi-user. */
+	struct FScopeMultiUserReceiveCVar
+	{
+		FScopeMultiUserReceiveCVar() = delete;
+		FScopeMultiUserReceiveCVar(const FScopeMultiUserReceiveCVar&) = delete;
+		FScopeMultiUserReceiveCVar& operator=(const FScopeMultiUserReceiveCVar&) = delete;
+
+		FScopeMultiUserReceiveCVar(FString InCommand, TMap<FString, int32>& InTable)
+			: TrackedCommand(MoveTemp(InCommand)),
+			  InboundCommandTable(InTable)
+		{
+			if (int32* Value = InboundCommandTable.Find(TrackedCommand))
+			{
+				*Value = *Value + 1;
+			}
+			else
+			{
+				int32& NewValue = InboundCommandTable.Add(TrackedCommand);
+				NewValue = 0;
+			}
+		}
+
+		~FScopeMultiUserReceiveCVar()
+		{
+			if (int32* Value = InboundCommandTable.Find(TrackedCommand))
+			{
+				*Value = *Value - 1;
+				if (*Value == 0)
+				{
+					InboundCommandTable.Remove(TrackedCommand);
+				}
+			}
+		}
+
+		FString TrackedCommand;
+		TMap<FString, int32>& InboundCommandTable;
+	};
+
+	void OnFEngineLoopInitComplete();
+
+	void RegisterMenuItem();
+	void RegisterProjectSettings() const;
+
+	void OnConsoleVariableChanged(IConsoleVariable* ChangedVariable);
+	/** In the event a console object is unregistered, this failsafe callback will clean up the associated list item and command info object. */
+	void OnDetectConsoleObjectUnregistered(FString CommandName);
+
+	TObjectPtr<UConsoleVariablesAsset> AllocateTransientPreset(const FName DesiredName) const;
+	void CreateEditingPresets();
+
+	TSharedRef<SDockTab> SpawnMainPanelTab(const FSpawnTabArgs& Args);
+
+	static void OpenConsoleVariablesEditor();
+
+	static const FName ConsoleVariablesToolkitPanelTabId;
+
+	/** Lives for as long as the module is loaded. */
+	TSharedPtr<FConsoleVariablesEditorMainPanel> MainPanel;
+
+	/** Transient preset that's being edited so we don't affect the reference asset unless we save it */
+	TObjectPtr<UConsoleVariablesAsset> EditingPresetAsset = nullptr;
+	/** Transient preset that tracks variables that match the search criteria */
+	TObjectPtr<UConsoleVariablesAsset> EditingGlobalSearchAsset = nullptr;
+
+	/** All tracked variables and their default, startup, and current values */
+	TArray<TSharedPtr<FConsoleVariablesEditorCommandInfo>> ConsoleObjectsMainReference;
+
+	/**
+	 * A map of commands invoked by Multi-user.  Each command has a corresponding reference count.
+	 * When this node receives a command from another node, it will be added to this list.  This is prevent
+	 * a remote cvar change creating a ping/pong cvar updates effect between nodes.
+	 */
+	TMap<FString, int32> CommandsReceivedFromMultiUser;
+
+	/* Have we warned the user about PIE and Console Variable Editor */
+	bool bHaveWarnedAboutPIE = false;
+};
