Make log warning strictly informational.

diff --git a/Engine/Plugins/Editor/ConsoleVariablesEditor/Source/ConsoleVariablesEditor/Private/ConsoleVariablesEditorModule.cpp b/Engine/Plugins/Editor/ConsoleVariablesEditor/Source/ConsoleVariablesEditor/Private/ConsoleVariablesEditorModule.cpp
--- a/Engine/Plugins/Editor/ConsoleVariablesEditor/Source/ConsoleVariablesEditor/Private/ConsoleVariablesEditorModule.cpp
+++ b/Engine/Plugins/Editor/ConsoleVariablesEditor/Source/ConsoleVariablesEditor/Private/ConsoleVariablesEditorModule.cpp
@@ -264,7 +264,7 @@ void FConsoleVariablesEditorModule::SendMultiUserConsoleVariableChange(ERemoteCV
 {
 	if (GEditor && GEditor->IsPlaySessionInProgress() && !bHaveWarnedAboutPIE)
 	{
-		UE_LOG(LogConsoleVariablesEditor, Warning, TEXT("%hs: Play In Editor is about to start or has started; Multi-User Cvar sync is suspended during PIE."), __FUNCTION__);
+		UE_LOG(LogConsoleVariablesEditor, Display, TEXT("%hs: Play In Editor is about to start or has started; Multi-User Cvar sync is suspended during PIE."), __FUNCTION__);
 		bHaveWarnedAboutPIE = true;
 		return;
 	}
