ConsoleVariablesEditor: Make the autocomplete widget public, and added a few features.

MRQ needs the same autocomplete widget, which is why it's being made public.

Added some settings on the widget for MRQ purposes (as well as keeping it compatible w/ existing CVE needs):
* Allow the input to be enabled/disabled.
* Allow overrides for font and hint text.
* Text can be provided, and optionally cleared upon commit.
* Provide TextChanged/Committed events.
* The widget can be hidden when focus is lost.

The input will also display a tooltip w/ the command's help text.

diff --git a/Engine/Plugins/Editor/ConsoleVariablesEditor/Source/ConsoleVariablesEditor/Private/Views/MainPanel/SConsoleVariablesEditorCustomConsoleInputBox.cpp b/Engine/Plugins/Editor/ConsoleVariablesEditor/Source/ConsoleVariablesEditor/Private/Views/MainPanel/SConsoleVariablesEditorCustomConsoleInputBox.cpp
deleted file mode 100644
diff --git a/Engine/Plugins/Editor/ConsoleVariablesEditor/Source/ConsoleVariablesEditor/Private/Views/MainPanel/SConsoleVariablesEditorCustomConsoleInputBox.h b/Engine/Plugins/Editor/ConsoleVariablesEditor/Source/ConsoleVariablesEditor/Private/Views/MainPanel/SConsoleVariablesEditorCustomConsoleInputBox.h
deleted file mode 100644
diff --git a/Engine/Plugins/Editor/ConsoleVariablesEditor/Source/ConsoleVariablesEditor/Private/Views/MainPanel/SConsoleVariablesEditorMainPanel.cpp b/Engine/Plugins/Editor/ConsoleVariablesEditor/Source/ConsoleVariablesEditor/Private/Views/MainPanel/SConsoleVariablesEditorMainPanel.cpp
--- a/Engine/Plugins/Editor/ConsoleVariablesEditor/Source/ConsoleVariablesEditor/Private/Views/MainPanel/SConsoleVariablesEditorMainPanel.cpp
+++ b/Engine/Plugins/Editor/ConsoleVariablesEditor/Source/ConsoleVariablesEditor/Private/Views/MainPanel/SConsoleVariablesEditorMainPanel.cpp
@@ -9,7 +9,7 @@
 #include "MultiUser/ConsoleVariableSyncData.h"
 #include "Views/List/ConsoleVariablesEditorList.h"
 #include "Views/MainPanel/ConsoleVariablesEditorMainPanel.h"
-#include "Views/MainPanel/SConsoleVariablesEditorCustomConsoleInputBox.h"
+#include "Views/Widgets/SConsoleVariablesEditorCustomConsoleInputBox.h"
 
 #include "ContentBrowserModule.h"
 #include "DetailsViewArgs.h"
@@ -56,8 +56,9 @@ void SConsoleVariablesEditorMainPanel::Construct(
 		.AutoHeight()
 		.Padding(FMargin(8.f, 0.f, 8.f, 0.f))
 		[
-			SAssignNew(AddConsoleObjectInputBoxPtr, SConsoleVariablesEditorCustomConsoleInputBox, SharedThis(this))
+			SAssignNew(AddConsoleObjectInputBoxPtr, SConsoleVariablesEditorCustomConsoleInputBox)
 			.Visibility(EVisibility::Collapsed)
+			.OnTextCommitted(this, &SConsoleVariablesEditorMainPanel::ValidateConsoleInputAndAddToCurrentPreset)
 		]
 
 		+SVerticalBox::Slot()
@@ -90,9 +91,9 @@ SConsoleVariablesEditorMainPanel::~SConsoleVariablesEditorMainPanel()
 	MultiUserDetailsBox.Reset();
 }
 
-FReply SConsoleVariablesEditorMainPanel::ValidateConsoleInputAndAddToCurrentPreset(const FText& CommittedText) const
+void SConsoleVariablesEditorMainPanel::ValidateConsoleInputAndAddToCurrentPreset(const FText& CommittedText) const
 {
-	return MainPanel.Pin()->ValidateConsoleInputAndAddToCurrentPreset(CommittedText);
+	MainPanel.Pin()->ValidateConsoleInputAndAddToCurrentPreset(CommittedText);
 }
 
 void SConsoleVariablesEditorMainPanel::RefreshMultiUserDetails() const
diff --git a/Engine/Plugins/Editor/ConsoleVariablesEditor/Source/ConsoleVariablesEditor/Private/Views/MainPanel/SConsoleVariablesEditorMainPanel.h b/Engine/Plugins/Editor/ConsoleVariablesEditor/Source/ConsoleVariablesEditor/Private/Views/MainPanel/SConsoleVariablesEditorMainPanel.h
--- a/Engine/Plugins/Editor/ConsoleVariablesEditor/Source/ConsoleVariablesEditor/Private/Views/MainPanel/SConsoleVariablesEditorMainPanel.h
+++ b/Engine/Plugins/Editor/ConsoleVariablesEditor/Source/ConsoleVariablesEditor/Private/Views/MainPanel/SConsoleVariablesEditorMainPanel.h
@@ -27,7 +27,7 @@ public:
 	virtual ~SConsoleVariablesEditorMainPanel() override;
 
 	/** A handler for when text is committed to the Console Input Box that appears when clicking the "Add Console Variable" button */
-	FReply ValidateConsoleInputAndAddToCurrentPreset(const FText& CommittedText) const;
+	void ValidateConsoleInputAndAddToCurrentPreset(const FText& CommittedText) const;
 
 	void RefreshMultiUserDetails() const;
 	void ToggleMultiUserDetails(ECheckBoxState CheckState) const;
diff --git a/Engine/Plugins/Editor/ConsoleVariablesEditor/Source/ConsoleVariablesEditor/Public/Views/Widgets/SConsoleVariablesEditorCustomConsoleInputBox.cpp b/Engine/Plugins/Editor/ConsoleVariablesEditor/Source/ConsoleVariablesEditor/Public/Views/Widgets/SConsoleVariablesEditorCustomConsoleInputBox.cpp
new file mode 100644
--- /dev/null
+++ b/Engine/Plugins/Editor/ConsoleVariablesEditor/Source/ConsoleVariablesEditor/Public/Views/Widgets/SConsoleVariablesEditorCustomConsoleInputBox.cpp
@@ -0,0 +1,449 @@
+// Copyright Epic Games, Inc. All Rights Reserved.
+
+#include "SConsoleVariablesEditorCustomConsoleInputBox.h"
+
+#include "ConsoleVariablesEditorCommandInfo.h"
+#include "ConsoleVariablesEditorLog.h"
+#include "ConsoleVariablesEditorModule.h"
+#include "Views/List/ConsoleVariablesEditorList.h"
+#include "Views/MainPanel/SConsoleVariablesEditorMainPanel.h"
+
+#include "Framework/Application/SlateApplication.h"
+#include "OutputLogModule.h"
+#include "Styling/AppStyle.h"
+#include "Styling/StyleColors.h"
+#include "Views/Widgets/SConsoleVariablesEditorTooltipWidget.h"
+#include "Widgets/Input/SEditableTextBox.h"
+#include "Widgets/Layout/SBorder.h"
+#include "Widgets/Layout/SBox.h"
+#include "Widgets/Text/STextBlock.h"
+
+#define LOCTEXT_NAMESPACE "ConsoleVariablesEditor"
+
+void SConsoleVariablesEditorCustomConsoleInputBox::Construct(
+	const FArguments& InArgs)
+{
+	bHideOnFocusLost = InArgs._HideOnFocusLost;
+	bClearOnCommit = InArgs._ClearOnCommit;
+	IsEnabledAttribute = InArgs._IsEnabled;
+	Font = InArgs._Font;
+	HintText = InArgs._HintText;
+	TextAttribute = InArgs._Text;
+	OnTextCommitted = InArgs._OnTextCommitted;
+	OnTextChanged = InArgs._OnTextChanged;
+	
+	ChildSlot
+	[
+		SAssignNew( SuggestionBox, SMenuAnchor )
+		.Method(GIsEditor ? EPopupMethod::CreateNewWindow : EPopupMethod::UseCurrentWindow)
+		.Placement( MenuPlacement_BelowAnchor )
+		[
+			SAssignNew(InputText, SEditableTextBox)
+			.IsEnabled(IsEnabledAttribute)
+			.Text(TextAttribute)
+			.ToolTipText_Lambda([this]()
+			{
+				const FString SanitizedCommand = GetSanitizedCommand(InputText->GetText().ToString());
+
+				FConsoleVariablesEditorModule& ConsoleVariablesEditorModule = FConsoleVariablesEditorModule::Get();
+				const TWeakPtr<FConsoleVariablesEditorCommandInfo> MatchingCommand = ConsoleVariablesEditorModule.FindCommandInfoByName(SanitizedCommand);
+				if (MatchingCommand.IsValid())
+				{
+					return FText::FromString(MatchingCommand.Pin()->GetHelpText());
+				}
+
+				return FText();
+			})
+			.Font(Font.HasValidFont() ? Font : FAppStyle::Get().GetWidgetStyle<FTextBlockStyle>("Log.Normal").Font)
+			.HintText(!HintText.IsEmpty() ? HintText : LOCTEXT("ConsoleCommandExecutorHintText", "Enter Console Command"))
+			.OnTextChanged(this, &SConsoleVariablesEditorCustomConsoleInputBox::OnInputTextChanged)
+			.OnKeyCharHandler(this, &SConsoleVariablesEditorCustomConsoleInputBox::OnKeyCharHandler)
+			.OnKeyDownHandler(this, &SConsoleVariablesEditorCustomConsoleInputBox::OnKeyDownHandler)
+			.OnTextCommitted_Lambda([this] (const FText& InText, const ETextCommit::Type CommitType)
+			{
+				// Hide widget when focus is lost
+				if (CommitType == ETextCommit::OnUserMovedFocus)
+				{
+					// If the newly focused widget is the suggestion list, we want the suggestion list to handle what's next
+					if (FSlateApplication::Get().GetUserFocusedWidget(0) != SuggestionListView &&
+						FSlateApplication::Get().GetUserFocusedWidget(0) != InputText)
+					{
+						SuggestionBox->SetIsOpen(false);
+
+						if (bHideOnFocusLost)
+						{
+							SetVisibility(EVisibility::Collapsed);
+						}
+					}
+				}
+			})
+			.ClearKeyboardFocusOnCommit(false)
+		
+		]
+		.MenuContent
+		(
+			SNew(SBorder)
+			.BorderImage(FAppStyle::GetBrush("Menu.Background"))
+			.Padding( FMargin(2) )
+			[
+				SNew(SBox)
+				.HeightOverride(250) // avoids flickering, ideally this would be adaptive to the content without flickering
+				.MinDesiredWidth(300)
+				.MaxDesiredWidth_Lambda([this] ()
+				{
+					// Limit the width of the suggestions list to the work area that this widget currently resides on
+					const FSlateRect WidgetRect(
+						GetCachedGeometry().GetAbsolutePosition(),
+						GetCachedGeometry().GetAbsolutePosition() + GetCachedGeometry().GetAbsoluteSize());
+					const FSlateRect WidgetWorkArea = FSlateApplication::Get().GetWorkArea(WidgetRect);
+					return FMath::Max(300.0f, WidgetWorkArea.GetSize().X - 12.0f);
+				})
+				[
+					SAssignNew(SuggestionListView, SListView< TSharedPtr<FString> >)
+					.ListItemsSource(&Suggestions.SuggestionsList)
+					.SelectionMode( ESelectionMode::Single )							// Ideally the mouse over would not highlight while keyboard controls the UI
+					.OnGenerateRow_Lambda([this] (TSharedPtr<FString> Text, const TSharedRef<STableViewBase>& OwnerTable)
+					{
+						check(Text.IsValid());
+
+						const FString SanitizedText = GetSanitizedCommand(*Text);
+						const FText DisplayText = FText::FromString(SanitizedText);
+						
+						FConsoleVariablesEditorModule& ConsoleVariablesEditorModule = FConsoleVariablesEditorModule::Get();
+
+						TSharedPtr<IToolTip> ToolTip;
+
+						if (const TWeakPtr<FConsoleVariablesEditorCommandInfo> MatchingCommand =
+							ConsoleVariablesEditorModule.FindCommandInfoByName(SanitizedText); MatchingCommand.IsValid())
+						{
+							ToolTip = SConsoleVariablesEditorTooltipWidget::MakeTooltip(
+								SanitizedText,
+								MatchingCommand.Pin()->GetHelpText());
+						}
+
+						return
+							SNew(STableRow< TSharedPtr<FString> >, OwnerTable)
+							[
+								SNew(STextBlock)
+								.Text(DisplayText)
+								.TextStyle(FAppStyle::Get(), "Log.Normal")
+								.HighlightText(Suggestions.SuggestionsHighlight)
+								.ColorAndOpacity(FSlateColor::UseForeground())
+								.ToolTip(ToolTip.IsValid() ? ToolTip : nullptr)
+							];
+					})
+					.OnSelectionChanged_Lambda([this] (TSharedPtr<FString> NewValue, ESelectInfo::Type SelectInfo)
+					{
+						if(bIgnoreUIUpdate)
+						{
+							return;
+						}
+						
+						Suggestions.SelectedSuggestion = Suggestions.SuggestionsList.IndexOfByPredicate([&NewValue](const TSharedPtr<FString>& InSuggestion)
+						{
+							return InSuggestion == NewValue;
+						});
+
+						MarkActiveSuggestion();
+
+						// Just close the suggestion box and add the marked suggestion to the list and close up
+						if( SelectInfo == ESelectInfo::OnMouseClick )
+						{
+							CommitInput();
+
+							if (bHideOnFocusLost)
+							{
+								SetVisibility(EVisibility::Collapsed);
+							}
+						}
+					})
+					.ItemHeight(18)
+				]
+			]
+		)
+	];
+}
+
+SConsoleVariablesEditorCustomConsoleInputBox::~SConsoleVariablesEditorCustomConsoleInputBox()
+{
+	InputText.Reset();
+	SuggestionBox.Reset();
+	SuggestionListView.Reset();
+}
+
+FReply SConsoleVariablesEditorCustomConsoleInputBox::OnPreviewKeyDown(const FGeometry& MyGeometry, const FKeyEvent& KeyEvent)
+{
+	if(SuggestionBox->IsOpen())
+	{
+		if(KeyEvent.GetKey() == EKeys::Up || KeyEvent.GetKey() == EKeys::Down)
+		{
+			Suggestions.StepSelectedSuggestion(KeyEvent.GetKey() == EKeys::Up ? -1 : +1);
+			MarkActiveSuggestion();
+
+			return FReply::Handled();
+		}
+		else if (KeyEvent.GetKey() == EKeys::Tab)
+		{
+			if (Suggestions.HasSuggestions())
+			{
+				if (Suggestions.HasSelectedSuggestion())
+				{
+					Suggestions.StepSelectedSuggestion(KeyEvent.IsShiftDown() ? -1 : +1);
+				}
+				else
+				{
+					Suggestions.SelectedSuggestion = 0;
+				}
+				MarkActiveSuggestion();
+			}
+			
+			return FReply::Handled();
+		}
+		else if (KeyEvent.GetKey() == EKeys::Escape)
+		{
+			SuggestionBox->SetIsOpen(false);
+			return FReply::Handled();
+		}
+	}
+	else
+	{
+		if(KeyEvent.GetKey() == EKeys::Up)
+		{
+			// If the command field isn't empty we need you to have pressed Control+Up to summon the history (to make sure you're not just using caret navigation)
+			const bool bShowHistory = InputText->GetText().IsEmpty() || KeyEvent.IsControlDown();
+			if (bShowHistory)
+			{
+				TArray<FString> History;
+				IConsoleManager::Get().GetConsoleHistory(TEXT(""), History);
+				
+				SetSuggestions(History, FText::GetEmpty());
+				
+				if(Suggestions.HasSuggestions())
+				{
+					Suggestions.StepSelectedSuggestion(-1);
+					MarkActiveSuggestion();
+				}
+			}
+
+			// Need to always handle this for single-line controls to avoid them invoking widget navigation
+			if (bShowHistory)
+			{
+				return FReply::Handled();
+			}
+		}
+	}
+
+	return FReply::Unhandled();
+}
+
+bool SConsoleVariablesEditorCustomConsoleInputBox::TakeKeyboardFocus() const
+{
+	check(InputText.IsValid());
+	
+	return FSlateApplication::Get().SetKeyboardFocus(InputText.ToSharedRef());
+}
+
+void SConsoleVariablesEditorCustomConsoleInputBox::OnInputTextChanged(const FText& InText)
+{
+	if(bIgnoreUIUpdate)
+	{
+		return;
+	}
+	
+	const FString& InputTextStr = InText.ToString();
+	if(!InputTextStr.IsEmpty())
+	{
+		TArray<FString> AutoCompleteList;
+		
+		auto OnConsoleVariable = [&AutoCompleteList](const TCHAR *Name, IConsoleObject* CVar)
+		{
+#if (UE_BUILD_SHIPPING || UE_BUILD_TEST)
+			if (CVar->TestFlags(ECVF_Cheat))
+			{
+				return;
+			}
+#endif // (UE_BUILD_SHIPPING || UE_BUILD_TEST)
+			if (CVar->TestFlags(ECVF_Unregistered))
+			{
+				return;
+			}
+
+			AutoCompleteList.Add(Name);
+		};
+
+		IConsoleManager::Get().ForEachConsoleObjectThatContains(FConsoleObjectVisitor::CreateLambda(OnConsoleVariable), *InputTextStr);
+
+		AutoCompleteList.Sort([InputTextStr](const FString& A, const FString& B)
+		{ 
+			if (A.StartsWith(InputTextStr))
+			{
+				if (!B.StartsWith(InputTextStr))
+				{
+					return true;
+				}
+			}
+			else
+			{
+				if (B.StartsWith(InputTextStr))
+				{
+					return false;
+				}
+			}
+
+			return A < B;
+
+		});
+
+
+		SetSuggestions(AutoCompleteList, InText);
+	}
+	else
+	{
+		ClearSuggestions();
+	}
+
+	OnTextChanged.ExecuteIfBound(InText);
+}
+
+FReply SConsoleVariablesEditorCustomConsoleInputBox::OnKeyCharHandler(const FGeometry& MyGeometry, const FCharacterEvent& InCharacterEvent) const
+{
+	// Intercept tab if used for auto-complete
+	if (InCharacterEvent.GetCharacter() == '\t')
+	{
+		return FReply::Handled();
+	}
+
+	if (InCharacterEvent.GetModifierKeys().AnyModifiersDown() && InCharacterEvent.GetCharacter() == ' ')
+	{	
+		// Ignore space bar + a modifier key.  It should not type a space as this is used by other keyboard shortcuts
+		return FReply::Handled();
+	}
+
+	return FReply::Unhandled();
+}
+
+FReply SConsoleVariablesEditorCustomConsoleInputBox::OnKeyDownHandler(const FGeometry& MyGeometry, const FKeyEvent& KeyPressed)
+{
+	if (KeyPressed.GetKey().GetFName() == TEXT("Enter"))
+	{
+		CommitInput();
+	}
+
+	return FReply::Unhandled();
+}
+
+void SConsoleVariablesEditorCustomConsoleInputBox::SetSuggestions(TArray<FString>& Elements, FText Highlight)
+{
+	FString SelectionText;
+	if (Suggestions.HasSelectedSuggestion())
+	{
+		SelectionText = *Suggestions.GetSelectedSuggestion();
+	}
+
+	Suggestions.Reset();
+	Suggestions.SuggestionsHighlight = Highlight;
+
+	for(int32 i = 0; i < Elements.Num(); ++i)
+	{
+		Suggestions.SuggestionsList.Add(MakeShared<FString>(Elements[i]));
+
+		if (Elements[i] == SelectionText)
+		{
+			Suggestions.SelectedSuggestion = i;
+		}
+	}
+	SuggestionListView->RequestListRefresh();
+
+	if(Suggestions.HasSuggestions())
+	{
+		// Ideally if the selection box is open the output window is not changing it's window title (flickers)
+		SuggestionBox->SetIsOpen(true, false);
+		if (Suggestions.HasSelectedSuggestion())
+		{
+			SuggestionListView->RequestScrollIntoView(Suggestions.GetSelectedSuggestion());
+		}
+		else
+		{
+			SuggestionListView->ScrollToTop();
+		}
+	}
+	else
+	{
+		SuggestionBox->SetIsOpen(false);
+	}
+}
+
+void SConsoleVariablesEditorCustomConsoleInputBox::MarkActiveSuggestion()
+{
+	bIgnoreUIUpdate = true;
+	if (Suggestions.HasSelectedSuggestion())
+	{
+		TSharedPtr<FString> SelectedSuggestion = Suggestions.GetSelectedSuggestion();
+
+		SuggestionListView->SetSelection(SelectedSuggestion);
+		SuggestionListView->RequestScrollIntoView(SelectedSuggestion);	// Ideally this would only scroll if outside of the view
+
+		const FText SelectedText = FText::FromString(*SelectedSuggestion);
+		OnTextChanged.ExecuteIfBound(SelectedText);
+
+		if (!TextAttribute.IsBound())
+		{
+			// Update the text input manually if it's not bound to an attribute
+			InputText->SetText(SelectedText);
+		}
+		else
+		{
+			// The input seems to need focus when the suggestions box is open in order to receive updates. This
+			// is only a problem if receiving updates from the bound Text attribute.
+			TakeKeyboardFocus();
+		}
+	}
+	else
+	{
+		SuggestionListView->ClearSelection();
+	}
+	bIgnoreUIUpdate = false;
+}
+
+void SConsoleVariablesEditorCustomConsoleInputBox::ClearSuggestions()
+{
+	SuggestionBox->SetIsOpen(false);
+	Suggestions.Reset();
+}
+
+void SConsoleVariablesEditorCustomConsoleInputBox::CommitInput()
+{
+	const FText CommittedText = InputText->GetText();
+	bIgnoreUIUpdate = true;
+
+	if (bClearOnCommit)
+	{
+		OnTextChanged.ExecuteIfBound(FText::GetEmpty());
+
+		if (!TextAttribute.IsBound())
+		{
+			InputText->SetText(FText::GetEmpty());
+		}
+	}
+	else
+	{
+		OnTextChanged.ExecuteIfBound(CommittedText);
+	}
+
+	OnTextCommitted.ExecuteIfBound(CommittedText);
+
+	SuggestionBox->SetIsOpen( false );
+	bIgnoreUIUpdate = false;
+}
+
+FString SConsoleVariablesEditorCustomConsoleInputBox::GetSanitizedCommand(const FString& InCommand) const
+{
+	FString SanitizedText = InCommand;
+	SanitizedText.ReplaceInline(TEXT("\r\n"), TEXT("\n"), ESearchCase::CaseSensitive);
+	SanitizedText.ReplaceInline(TEXT("\r"), TEXT(" "), ESearchCase::CaseSensitive);
+	SanitizedText.ReplaceInline(TEXT("\n"), TEXT(" "), ESearchCase::CaseSensitive);
+
+	return SanitizedText;
+}
+
+#undef LOCTEXT_NAMESPACE
diff --git a/Engine/Plugins/Editor/ConsoleVariablesEditor/Source/ConsoleVariablesEditor/Public/Views/Widgets/SConsoleVariablesEditorCustomConsoleInputBox.h b/Engine/Plugins/Editor/ConsoleVariablesEditor/Source/ConsoleVariablesEditor/Public/Views/Widgets/SConsoleVariablesEditorCustomConsoleInputBox.h
new file mode 100644
--- /dev/null
+++ b/Engine/Plugins/Editor/ConsoleVariablesEditor/Source/ConsoleVariablesEditor/Public/Views/Widgets/SConsoleVariablesEditorCustomConsoleInputBox.h
@@ -0,0 +1,148 @@
+// Copyright Epic Games, Inc. All Rights Reserved.
+
+#pragma once
+
+#include "Widgets/SCompoundWidget.h"
+#include "Widgets/Input/SMenuAnchor.h"
+#include "Widgets/Views/SListView.h"
+
+class SEditableTextBox;
+
+class CONSOLEVARIABLESEDITOR_API SConsoleVariablesEditorCustomConsoleInputBox : public SCompoundWidget
+{
+public:
+	DECLARE_DELEGATE_OneParam(FOnTextCommitted, const FText&);
+	DECLARE_DELEGATE_OneParam(FOnTextChanged, const FText&);
+
+	SLATE_BEGIN_ARGS(SConsoleVariablesEditorCustomConsoleInputBox)
+	{}
+		/** If true, hide the input box when focus is lost. */
+		SLATE_ARGUMENT(bool, HideOnFocusLost)
+
+		/** If true, clear the input box when text is committed. */
+		SLATE_ARGUMENT(bool, ClearOnCommit)
+
+		/** The enable state of the input box. */
+		SLATE_ATTRIBUTE(bool, IsEnabled)
+
+		/** Font override for the input box. */
+		SLATE_ARGUMENT(FSlateFontInfo, Font)
+
+		/** Hint text override for the input box. */
+		SLATE_ARGUMENT(FText, HintText)
+
+		/** The text to show in the input box. */
+		SLATE_ATTRIBUTE(FText, Text)
+
+		/** Called when text in the input box is committed. */
+		SLATE_EVENT(FOnTextCommitted, OnTextCommitted)
+
+		/** Called when text in the input box has changed. */
+		SLATE_EVENT(FOnTextChanged, OnTextChanged)
+	SLATE_END_ARGS()
+	
+	struct FSuggestions
+	{
+		FSuggestions()
+			: SelectedSuggestion(INDEX_NONE)
+		{
+		}
+
+		void Reset()
+		{
+			SelectedSuggestion = INDEX_NONE;
+			SuggestionsList.Reset();
+			SuggestionsHighlight = FText::GetEmpty();
+		}
+
+		bool HasSuggestions() const
+		{
+			return SuggestionsList.Num() > 0;
+		}
+
+		bool HasSelectedSuggestion() const
+		{
+			return SuggestionsList.IsValidIndex(SelectedSuggestion);
+		}
+
+		void StepSelectedSuggestion(const int32 Step)
+		{
+			SelectedSuggestion += Step;
+			if (SelectedSuggestion < 0)
+			{
+				SelectedSuggestion = SuggestionsList.Num() - 1;
+			}
+			else if (SelectedSuggestion >= SuggestionsList.Num())
+			{
+				SelectedSuggestion = 0;
+			}
+		}
+
+		TSharedPtr<FString> GetSelectedSuggestion() const
+		{
+			return SuggestionsList.IsValidIndex(SelectedSuggestion) ? SuggestionsList[SelectedSuggestion] : nullptr;
+		}
+
+		/** INDEX_NONE if not set, otherwise index into SuggestionsList */
+		int32 SelectedSuggestion;
+
+		/** All log messages stored in this widget for the list view */
+		TArray<TSharedPtr<FString>> SuggestionsList;
+
+		/** Highlight text to use for the suggestions list */
+		FText SuggestionsHighlight;
+	};
+
+	void Construct(const FArguments& InArgs);
+
+	virtual ~SConsoleVariablesEditorCustomConsoleInputBox() override;
+
+	//~ Begin SWidget Interface
+	virtual FReply OnPreviewKeyDown(const FGeometry& MyGeometry, const FKeyEvent& KeyEvent) override;
+	//~ End SWidget Interface
+
+	bool TakeKeyboardFocus() const;
+
+	void OnInputTextChanged(const FText& InText);
+
+	FReply OnKeyCharHandler(const FGeometry& MyGeometry, const FCharacterEvent& InCharacterEvent) const;
+
+	FReply OnKeyDownHandler(const FGeometry& MyGeometry, const FKeyEvent& KeyPressed);
+	
+	void SetSuggestions(TArray<FString>& Elements, FText Highlight);
+
+	void MarkActiveSuggestion();
+
+	void ClearSuggestions();
+	
+	void CommitInput();
+
+private:
+	FString GetSanitizedCommand(const FString& InCommand) const;
+
+private:
+	
+	/** A reference to the actual text box inside ConsoleInput */
+	TSharedPtr<SEditableTextBox> InputText;
+
+	/** history / auto completion elements */
+	TSharedPtr<SMenuAnchor> SuggestionBox;
+
+	/** The list view for showing all log messages. Should be replaced by a full text editor */
+	TSharedPtr<SListView<TSharedPtr<FString>>> SuggestionListView;
+
+	/** Active list of suggestions */
+	FSuggestions Suggestions;
+
+	/** to prevent recursive calls in UI callback */
+	bool bIgnoreUIUpdate = false;
+
+	bool bHideOnFocusLost = true;
+	bool bClearOnCommit = true;
+	TAttribute<bool> IsEnabledAttribute;
+	FSlateFontInfo Font;
+	FText HintText;
+	TAttribute<FText> TextAttribute;
+	FOnTextCommitted OnTextCommitted;
+	FOnTextChanged OnTextChanged;
+};
