CVar Editor: Bypass cvar change callback during automation tests.

#rnx

diff --git a/Engine/Plugins/Editor/ConsoleVariablesEditor/Source/ConsoleVariablesEditor/Private/ConsoleVariablesEditorModule.cpp b/Engine/Plugins/Editor/ConsoleVariablesEditor/Source/ConsoleVariablesEditor/Private/ConsoleVariablesEditorModule.cpp
--- a/Engine/Plugins/Editor/ConsoleVariablesEditor/Source/ConsoleVariablesEditor/Private/ConsoleVariablesEditorModule.cpp
+++ b/Engine/Plugins/Editor/ConsoleVariablesEditor/Source/ConsoleVariablesEditor/Private/ConsoleVariablesEditorModule.cpp
@@ -352,7 +352,8 @@ void FConsoleVariablesEditorModule::RegisterProjectSettings() const
 
 void FConsoleVariablesEditorModule::OnConsoleVariableChanged(IConsoleVariable* ChangedVariable)
 {
-	if (IsEngineExitRequested())
+	// Note: We disable tracking during automation tests to prevent FindConsoleObject() related performance warnings.
+	if (IsEngineExitRequested() || GIsAutomationTesting)
 	{
 		return;
 	}
