Fix crash in ~FManagerImpl not checking for nullptr

diff --git a/Engine/Plugins/Editor/ConsoleVariablesEditor/Source/ConsoleVariablesEditor/Private/MultiUser/ConsoleVariableSync.cpp b/Engine/Plugins/Editor/ConsoleVariablesEditor/Source/ConsoleVariablesEditor/Private/MultiUser/ConsoleVariableSync.cpp
--- a/Engine/Plugins/Editor/ConsoleVariablesEditor/Source/ConsoleVariablesEditor/Private/MultiUser/ConsoleVariableSync.cpp
+++ b/Engine/Plugins/Editor/ConsoleVariablesEditor/Source/ConsoleVariablesEditor/Private/MultiUser/ConsoleVariableSync.cpp
@@ -50,12 +50,13 @@ struct FManagerImpl
 	{
 		if (IConcertSyncClientModule::IsAvailable())
 		{
-			TSharedPtr<IConcertSyncClient> ConcertSyncClient = IConcertSyncClientModule::Get().GetClient(TEXT("MultiUser"));
 			Unregister();
-
-			IConcertClientRef ConcertClient = ConcertSyncClient->GetConcertClient();
-			ConcertClient->OnSessionStartup().RemoveAll(this);
-			ConcertClient->OnSessionShutdown().RemoveAll(this);
+			if (const TSharedPtr<IConcertSyncClient> ConcertSyncClient = IConcertSyncClientModule::Get().GetClient(TEXT("MultiUser")))
+			{
+				IConcertClientRef ConcertClient = ConcertSyncClient->GetConcertClient();
+				ConcertClient->OnSessionStartup().RemoveAll(this);
+				ConcertClient->OnSessionShutdown().RemoveAll(this);
+			}
 		}
 
 		UnregisterExtensions();
