Fix crash on shutdown with FConsoleVariablesEditorCommandInfo accessing stale deleted IConsoleVariable

diff --git a/Engine/Plugins/Editor/ConsoleVariablesEditor/Source/ConsoleVariablesEditor/Private/ConsoleVariablesEditorModule.cpp b/Engine/Plugins/Editor/ConsoleVariablesEditor/Source/ConsoleVariablesEditor/Private/ConsoleVariablesEditorModule.cpp
--- a/Engine/Plugins/Editor/ConsoleVariablesEditor/Source/ConsoleVariablesEditor/Private/ConsoleVariablesEditorModule.cpp
+++ b/Engine/Plugins/Editor/ConsoleVariablesEditor/Source/ConsoleVariablesEditor/Private/ConsoleVariablesEditorModule.cpp
@@ -54,9 +54,9 @@ void FConsoleVariablesEditorModule::ShutdownModule()
 	ConsoleObjectsMainReference.Empty();
 
 	// Unregister project settings
-	ISettingsModule& SettingsModule = FModuleManager::LoadModuleChecked<ISettingsModule>("Settings");
+	if (ISettingsModule* SettingsModule = FModuleManager::GetModulePtr<ISettingsModule>("Settings"))
 	{
-		SettingsModule.UnregisterSettings("Project", "Plugins", "Console Variables Editor");
+		SettingsModule->UnregisterSettings("Project", "Plugins", "Console Variables Editor");
 	}
 }
 
diff --git a/Engine/Plugins/Editor/ConsoleVariablesEditor/Source/ConsoleVariablesEditor/Public/ConsoleVariablesEditorCommandInfo.h b/Engine/Plugins/Editor/ConsoleVariablesEditor/Source/ConsoleVariablesEditor/Public/ConsoleVariablesEditorCommandInfo.h
--- a/Engine/Plugins/Editor/ConsoleVariablesEditor/Source/ConsoleVariablesEditor/Public/ConsoleVariablesEditorCommandInfo.h
+++ b/Engine/Plugins/Editor/ConsoleVariablesEditor/Source/ConsoleVariablesEditor/Public/ConsoleVariablesEditorCommandInfo.h
@@ -50,9 +50,12 @@ struct CONSOLEVARIABLESEDITOR_API FConsoleVariablesEditorCommandInfo
 	~FConsoleVariablesEditorCommandInfo()
 	{
 		OnDetectConsoleObjectUnregistered.Remove(OnDetectConsoleObjectUnregisteredHandle);
-		if (IConsoleVariable* AsVariable = GetConsoleVariablePtr())
+		if (!IsEngineExitRequested()) // CrashFix as GetConsoleVariablePtr() returns stale deleted IConsoleVariable during shut down
 		{
-			AsVariable->OnChangedDelegate().Remove(OnVariableChangedCallbackHandle);
+			if (IConsoleVariable* AsVariable = GetConsoleVariablePtr())
+			{
+				AsVariable->OnChangedDelegate().Remove(OnVariableChangedCallbackHandle);
+			}
 		}
 	}
 	
