Console Variables Editor:

-Check for session before sending/receiving MU cvar change

diff --git a/Engine/Plugins/Editor/ConsoleVariablesEditor/Source/ConsoleVariablesEditor/Private/ConsoleVariablesEditorModule.cpp b/Engine/Plugins/Editor/ConsoleVariablesEditor/Source/ConsoleVariablesEditor/Private/ConsoleVariablesEditorModule.cpp
--- a/Engine/Plugins/Editor/ConsoleVariablesEditor/Source/ConsoleVariablesEditor/Private/ConsoleVariablesEditorModule.cpp
+++ b/Engine/Plugins/Editor/ConsoleVariablesEditor/Source/ConsoleVariablesEditor/Private/ConsoleVariablesEditorModule.cpp
@@ -23,6 +23,27 @@
 
 #define LOCTEXT_NAMESPACE "FConsoleVariablesEditorModule"
 
+namespace UE::ConsoleVariablesEditor::ModuleUtils::Private
+{
+	/**
+	 * Warn the user once if attempting to sync Cvars in multi-user if in PIE.
+	 * @return Whether the user is in PIE.
+	 */
+	static bool CheckIsRunningPieAndWarnSyncDisabledInPie()
+	{
+		/* Have we warned the user about PIE in the Console Variables Editor? */
+		static bool bHaveWarnedAboutPIE = false;
+	
+		const bool bRunningInPIE =  GEditor && GEditor->IsPlaySessionInProgress();
+		if (bRunningInPIE && !bHaveWarnedAboutPIE)
+		{
+			UE_LOG(LogConsoleVariablesEditor, Display, TEXT("Play In Editor is about to start or has started; Multi-User Cvar sync is suspended during PIE."));
+			bHaveWarnedAboutPIE = bRunningInPIE;
+		}
+		return bRunningInPIE;
+	}
+}
+
 const FName FConsoleVariablesEditorModule::ConsoleVariablesToolkitPanelTabId(TEXT("ConsoleVariablesToolkitPanel"));
 
 FConsoleVariablesEditorModule& FConsoleVariablesEditorModule::Get()
@@ -262,30 +283,19 @@ bool FConsoleVariablesEditorModule::PopulateGlobalSearchAssetWithVariablesMatchi
 
 void FConsoleVariablesEditorModule::SendMultiUserConsoleVariableChange(ERemoteCVarChangeType InChangeType, const FString& InVariableName, const FString& InValueAsString)
 {
-	if (GEditor && GEditor->IsPlaySessionInProgress() && !bHaveWarnedAboutPIE)
+	using namespace UE::ConsoleVariablesEditor::ModuleUtils::Private; 
+	if (MainPanel->GetMultiUserManager().IsLocalUserInMultiUserSession() && !CheckIsRunningPieAndWarnSyncDisabledInPie())
 	{
-		UE_LOG(LogConsoleVariablesEditor, Display, TEXT("%hs: Play In Editor is about to start or has started; Multi-User Cvar sync is suspended during PIE."), __FUNCTION__);
-		bHaveWarnedAboutPIE = true;
-		return;
+		MainPanel->GetMultiUserManager().SendConsoleVariableChange(InChangeType, InVariableName, InValueAsString);
 	}
-
-	bHaveWarnedAboutPIE = false;
-	MainPanel->GetMultiUserManager().SendConsoleVariableChange(InChangeType, InVariableName, InValueAsString);
 }
 
 void FConsoleVariablesEditorModule::OnRemoteCvarChanged(ERemoteCVarChangeType InChangeType, const FString InName, const FString InValue)
 {
-	if (GEditor && GEditor->IsPlaySessionInProgress() && !bHaveWarnedAboutPIE)
-	{
-		UE_LOG(LogConsoleVariablesEditor, Warning, TEXT("%hs: Play In Editor is about to start or has started; Multi-User Cvar sync is suspended during PIE."), __FUNCTION__);
-		bHaveWarnedAboutPIE = true;
-		return;
-	}
-
-	bHaveWarnedAboutPIE = false;
-	UE_LOG(LogConsoleVariablesEditor, VeryVerbose, TEXT("Remote set console variable %s = %s"), *InName, *InValue);
-
-	if (GetDefault<UConcertCVarSynchronization>()->bSyncCVarTransactions)
+	using namespace UE::ConsoleVariablesEditor::ModuleUtils::Private; 
+	if (MainPanel->GetMultiUserManager().IsLocalUserInMultiUserSession() &&
+		!CheckIsRunningPieAndWarnSyncDisabledInPie() &&
+		GetDefault<UConcertCVarSynchronization>()->bSyncCVarTransactions)
 	{
 		FScopeMultiUserReceiveCVar CVarChange(*InName, CommandsReceivedFromMultiUser);
 
@@ -307,10 +317,7 @@ void FConsoleVariablesEditorModule::OnRemoteCvarChanged(ERemoteCVarChangeType In
 			);
 		}
 
-		if (MainPanel->GetEditorListMode() == FConsoleVariablesEditorList::EConsoleVariablesEditorListMode::Preset)
-		{
-			MainPanel->RebuildList();
-		}
+		MainPanel->OnRemoteCvarChanged();
 	}
 }
 
diff --git a/Engine/Plugins/Editor/ConsoleVariablesEditor/Source/ConsoleVariablesEditor/Private/MultiUser/ConsoleVariableSync.cpp b/Engine/Plugins/Editor/ConsoleVariablesEditor/Source/ConsoleVariablesEditor/Private/MultiUser/ConsoleVariableSync.cpp
--- a/Engine/Plugins/Editor/ConsoleVariablesEditor/Source/ConsoleVariablesEditor/Private/MultiUser/ConsoleVariableSync.cpp
+++ b/Engine/Plugins/Editor/ConsoleVariablesEditor/Source/ConsoleVariablesEditor/Private/MultiUser/ConsoleVariableSync.cpp
@@ -514,5 +514,20 @@ bool FManager::IsInitialized() const
 {
 	return Implementation->bIsInitialized;
 }
+
+bool FManager::IsSyncEnabled() const
+{
+	return Implementation->bIsEnabled;
+}
+
+bool FManager::IsSessionValid() const
+{
+	return Implementation->WeakSession.IsValid();
+}
+
+bool FManager::IsLocalUserInMultiUserSession()
+{
+	return IsSyncEnabled() && IsSessionValid();
+}
 };
 
diff --git a/Engine/Plugins/Editor/ConsoleVariablesEditor/Source/ConsoleVariablesEditor/Private/MultiUser/ConsoleVariableSync.h b/Engine/Plugins/Editor/ConsoleVariablesEditor/Source/ConsoleVariablesEditor/Private/MultiUser/ConsoleVariableSync.h
--- a/Engine/Plugins/Editor/ConsoleVariablesEditor/Source/ConsoleVariablesEditor/Private/MultiUser/ConsoleVariableSync.h
+++ b/Engine/Plugins/Editor/ConsoleVariablesEditor/Source/ConsoleVariablesEditor/Private/MultiUser/ConsoleVariableSync.h
@@ -56,6 +56,12 @@ struct FManager
 	void SetEnableMultiUserSupport(bool bIsEnabled);
 
 	bool IsInitialized() const;
+
+	bool IsSyncEnabled() const;
+
+	bool IsSessionValid() const;
+
+	bool IsLocalUserInMultiUserSession();
 	
 private:
 	TUniquePtr<FManagerImpl> Implementation;
diff --git a/Engine/Plugins/Editor/ConsoleVariablesEditor/Source/ConsoleVariablesEditor/Private/Views/MainPanel/ConsoleVariablesEditorMainPanel.cpp b/Engine/Plugins/Editor/ConsoleVariablesEditor/Source/ConsoleVariablesEditor/Private/Views/MainPanel/ConsoleVariablesEditorMainPanel.cpp
--- a/Engine/Plugins/Editor/ConsoleVariablesEditor/Source/ConsoleVariablesEditor/Private/Views/MainPanel/ConsoleVariablesEditorMainPanel.cpp
+++ b/Engine/Plugins/Editor/ConsoleVariablesEditor/Source/ConsoleVariablesEditor/Private/Views/MainPanel/ConsoleVariablesEditorMainPanel.cpp
@@ -96,6 +96,14 @@ void FConsoleVariablesEditorMainPanel::UpdateCachedValue(const FString& InConsol
 	}
 }
 
+void FConsoleVariablesEditorMainPanel::OnRemoteCvarChanged()
+{
+	if (GetEditorListMode() == FConsoleVariablesEditorList::EConsoleVariablesEditorListMode::Preset)
+	{
+		RebuildList();
+	}
+}
+
 FReply FConsoleVariablesEditorMainPanel::ValidateConsoleInputAndAddToCurrentPreset(const FText& CommittedText) const
 {
 	// Clear Search
diff --git a/Engine/Plugins/Editor/ConsoleVariablesEditor/Source/ConsoleVariablesEditor/Private/Views/MainPanel/ConsoleVariablesEditorMainPanel.h b/Engine/Plugins/Editor/ConsoleVariablesEditor/Source/ConsoleVariablesEditor/Private/Views/MainPanel/ConsoleVariablesEditorMainPanel.h
--- a/Engine/Plugins/Editor/ConsoleVariablesEditor/Source/ConsoleVariablesEditor/Private/Views/MainPanel/ConsoleVariablesEditorMainPanel.h
+++ b/Engine/Plugins/Editor/ConsoleVariablesEditor/Source/ConsoleVariablesEditor/Private/Views/MainPanel/ConsoleVariablesEditorMainPanel.h
@@ -29,6 +29,8 @@ public:
 
 	void UpdateCachedValue(const FString& InConsoleVar, const FString& Value);
 
+	void OnRemoteCvarChanged();
+
 	FReply ValidateConsoleInputAndAddToCurrentPreset(const FText& CommittedText) const;
 
 	/*
diff --git a/Engine/Plugins/Editor/ConsoleVariablesEditor/Source/ConsoleVariablesEditor/Public/ConsoleVariablesEditorModule.h b/Engine/Plugins/Editor/ConsoleVariablesEditor/Source/ConsoleVariablesEditor/Public/ConsoleVariablesEditorModule.h
--- a/Engine/Plugins/Editor/ConsoleVariablesEditor/Source/ConsoleVariablesEditor/Public/ConsoleVariablesEditorModule.h
+++ b/Engine/Plugins/Editor/ConsoleVariablesEditor/Source/ConsoleVariablesEditor/Public/ConsoleVariablesEditorModule.h
@@ -164,7 +164,4 @@ private:
 	 * a remote cvar change creating a ping/pong cvar updates effect between nodes.
 	 */
 	TMap<FString, int32> CommandsReceivedFromMultiUser;
-
-	/* Have we warned the user about PIE and Console Variable Editor */
-	bool bHaveWarnedAboutPIE = false;
 };
