Console Variables Editor:

-Fix: "Show" Filters Do Not Update Cvar List

diff --git a/Engine/Plugins/Editor/ConsoleVariablesEditor/Source/ConsoleVariablesEditor/Private/Views/List/SConsoleVariablesEditorList.cpp b/Engine/Plugins/Editor/ConsoleVariablesEditor/Source/ConsoleVariablesEditor/Private/Views/List/SConsoleVariablesEditorList.cpp
--- a/Engine/Plugins/Editor/ConsoleVariablesEditor/Source/ConsoleVariablesEditor/Private/Views/List/SConsoleVariablesEditorList.cpp
+++ b/Engine/Plugins/Editor/ConsoleVariablesEditor/Source/ConsoleVariablesEditor/Private/Views/List/SConsoleVariablesEditorList.cpp
@@ -697,13 +697,23 @@ void SConsoleVariablesEditorList::EvaluateIfRowsPassFilters(const bool bShouldRe
 	{
 		if (Row.IsValid() && Row->GetRowType() == FConsoleVariablesEditorListRow::SingleCommand)
 		{
-			auto Projection = [&Row](const TSharedRef<IConsoleVariablesEditorListFilter>& Filter)
+			// If any filters of this group are inactive, they automatically fail.
+			// For "Show" filters, this means another active filter must pass in order for the row to pass 'bPassesAnyOf'.
+			auto AnyOfProjection = [&Row](const TSharedRef<IConsoleVariablesEditorListFilter>& Filter)
 			{
-				return Filter->GetIsFilterActive() ? Filter->DoesItemPassFilter(Row) : true;
+				return Filter->GetIsFilterActive() && Filter->DoesItemPassFilter(Row);
+			};
+
+			// If any filters of this group are inactive, they automatically pass.
+			// For "Show Only" filters, only active filters will check the passing condition,
+			// and all must pass in order for the row to pass 'bPassesAllOf'.
+			auto AllOfProjection = [&Row](const TSharedRef<IConsoleVariablesEditorListFilter>& Filter)
+			{
+				return !Filter->GetIsFilterActive() || Filter->DoesItemPassFilter(Row);
 			};
 			
-			const bool bPassesAnyOf = Algo::AnyOf(MatchAnyOfFilters, Projection);
-			const bool bPassesAllOf = Algo::AllOf(MatchAllOfFilters, Projection);
+			const bool bPassesAnyOf = Algo::AnyOf(MatchAnyOfFilters, AnyOfProjection);
+			const bool bPassesAllOf = Algo::AllOf(MatchAllOfFilters, AllOfProjection);
 			
 			Row->SetDoesRowPassFilters(bPassesAnyOf && bPassesAllOf);
 		}
