Console Variables Editor:

-Fix: Console Variables Editor does not Execute "HighResShot" more than once when the Execute button is selected.

diff --git a/Engine/Plugins/Editor/ConsoleVariablesEditor/Source/ConsoleVariablesEditor/Private/Views/List/SConsoleVariablesEditorListValueInput.cpp b/Engine/Plugins/Editor/ConsoleVariablesEditor/Source/ConsoleVariablesEditor/Private/Views/List/SConsoleVariablesEditorListValueInput.cpp
--- a/Engine/Plugins/Editor/ConsoleVariablesEditor/Source/ConsoleVariablesEditor/Private/Views/List/SConsoleVariablesEditorListValueInput.cpp
+++ b/Engine/Plugins/Editor/ConsoleVariablesEditor/Source/ConsoleVariablesEditor/Private/Views/List/SConsoleVariablesEditorListValueInput.cpp
@@ -439,21 +439,16 @@ void SConsoleVariablesEditorListValueInput_Command::Construct(const FArguments&
 			{
 				check (Item.IsValid());
 				
-				const FString ValueAsString = InputText->GetText().ToString();
-				
 				if (const TSharedPtr<FConsoleVariablesEditorListRow> PinnedItem = Item.Pin())
-				{
-					// Only execute if this row is a non-value command OR
-					// it's a value-based command whose cached value differs from the input value
-					if (PinnedItem->GetCommandInfo().Pin()->ObjectType ==
-						FConsoleVariablesEditorCommandInfo::EConsoleObjectType::NullObject ||
-						(PinnedItem->GetCommandInfo().Pin()->ObjectType ==
-						FConsoleVariablesEditorCommandInfo::EConsoleObjectType::Command &&
-						!PinnedItem->GetCachedValue().Equals(ValueAsString)))
+				{					
+					if (const TSharedPtr<FConsoleVariablesEditorCommandInfo> PinnedCommand =
+						PinnedItem->GetCommandInfo().Pin())
 					{
-						PinnedItem->GetCommandInfo().Pin()->ExecuteCommand(ValueAsString);
+						const FString InputValueAsString = InputText->GetText().ToString();
+						
+						PinnedCommand->ExecuteCommand(InputValueAsString);
 
-						PinnedItem->SetCachedValue(ValueAsString);
+						PinnedItem->SetCachedValue(InputValueAsString);
 					}
 
 					return FReply::Handled();
