Crash fixes for FConsoleVariablesEditorModule

diff --git a/Engine/Plugins/Editor/ConsoleVariablesEditor/Source/ConsoleVariablesEditor/Private/ConsoleVariablesEditorCommandInfo.cpp b/Engine/Plugins/Editor/ConsoleVariablesEditor/Source/ConsoleVariablesEditor/Private/ConsoleVariablesEditorCommandInfo.cpp
--- a/Engine/Plugins/Editor/ConsoleVariablesEditor/Source/ConsoleVariablesEditor/Private/ConsoleVariablesEditorCommandInfo.cpp
+++ b/Engine/Plugins/Editor/ConsoleVariablesEditor/Source/ConsoleVariablesEditor/Private/ConsoleVariablesEditorCommandInfo.cpp
@@ -70,11 +70,6 @@ IConsoleObject* FConsoleVariablesEditorCommandInfo::GetConsoleObjectPtr()
 		ConsoleObjectPtr = IConsoleManager::Get().FindConsoleObject(*CommandKey);
 		TimeOfLastConsoleObjectRefresh = FDateTime::UtcNow();
 	}
-	// If the console object turns out to be unregistered, let interested parties know
-	if (ConsoleObjectPtr && ConsoleObjectPtr->TestFlags(ECVF_Unregistered))
-	{
-		OnDetectConsoleObjectUnregistered.Broadcast(Command);
-	}
 	return ConsoleObjectPtr;
 }
 
diff --git a/Engine/Plugins/Editor/ConsoleVariablesEditor/Source/ConsoleVariablesEditor/Private/ConsoleVariablesEditorModule.cpp b/Engine/Plugins/Editor/ConsoleVariablesEditor/Source/ConsoleVariablesEditor/Private/ConsoleVariablesEditorModule.cpp
--- a/Engine/Plugins/Editor/ConsoleVariablesEditor/Source/ConsoleVariablesEditor/Private/ConsoleVariablesEditorModule.cpp
+++ b/Engine/Plugins/Editor/ConsoleVariablesEditor/Source/ConsoleVariablesEditor/Private/ConsoleVariablesEditorModule.cpp
@@ -68,6 +68,12 @@ void FConsoleVariablesEditorModule::ShutdownModule()
 
 	FCoreDelegates::OnFEngineLoopInitComplete.RemoveAll(this);
 
+	if (OnConsoleObjectUnregisteredHandle.IsValid())
+	{
+		IConsoleManager::Get().OnConsoleObjectUnregistered().Remove(OnConsoleObjectUnregisteredHandle);
+		OnConsoleObjectUnregisteredHandle.Reset();
+	}
+
 	FConsoleVariablesEditorStyle::Shutdown();
 	
 	MainPanel.Reset();
@@ -153,8 +159,6 @@ void FConsoleVariablesEditorModule::QueryAndBeginTrackingConsoleVariables()
 				MakeShared<FConsoleVariablesEditorCommandInfo>(Key);
 			
 			Info->StartupSource = Info->GetSource();
-			Info->OnDetectConsoleObjectUnregisteredHandle = Info->OnDetectConsoleObjectUnregistered.AddRaw(
-				this, &FConsoleVariablesEditorModule::OnDetectConsoleObjectUnregistered);
 
 			if (IConsoleVariable* AsVariable = ConsoleObject->AsVariable())
 			{
@@ -165,6 +169,8 @@ void FConsoleVariablesEditorModule::QueryAndBeginTrackingConsoleVariables()
 			AddConsoleObjectCommandInfoToMainReference(Info);
 		}),
 		TEXT(""));
+
+	OnConsoleObjectUnregisteredHandle = IConsoleManager::Get().OnConsoleObjectUnregistered().AddRaw(this, &FConsoleVariablesEditorModule::OnConsoleObjectUnregistered);
 }
 
 TWeakPtr<FConsoleVariablesEditorCommandInfo> FConsoleVariablesEditorModule::FindCommandInfoByName(const FString& NameToSearch, ESearchCase::Type InSearchCase)
@@ -431,21 +437,47 @@ void FConsoleVariablesEditorModule::OnConsoleVariableChanged(IConsoleVariable* C
 	}
 }
 
-void FConsoleVariablesEditorModule::OnDetectConsoleObjectUnregistered(FString CommandName)
+void FConsoleVariablesEditorModule::OnConsoleObjectUnregistered(const TCHAR* InName, IConsoleObject* InConsoleObject)
 {
-	check(EditingPresetAsset);
+	if (ensure(InName))
+	{
+		if (ensure(EditingPresetAsset))
+		{
+			EditingPresetAsset->RemoveConsoleVariable(InName);
+		}
+	}
 
-	EditingPresetAsset->RemoveConsoleVariable(CommandName);
+	int32 FoundIndex = INDEX_NONE;
+	for (int32 i = 0; i < ConsoleObjectsMainReference.Num(); ++i)
+	{
+		if (ConsoleObjectsMainReference[i])
+		{
+			if (ConsoleObjectsMainReference[i]->ConsoleObjectPtr == InConsoleObject)
+			{
+				FoundIndex = i;
+				break;
+			}
+			// GetConsoleObjectPtr() is overwriting ConsoleObjectPtr so we need to search by command name as well
+			else if (InName && ConsoleObjectsMainReference[i]->Command == InName)
+			{
+				FoundIndex = i;
+				break;
+			}
+		}
+	}
 
-	if (MainPanel.IsValid())
+	if (FoundIndex != -1)
 	{
-		MainPanel->RefreshList();
+		// Null out deleted ConsoleObjectPtr for anything still holding reference to TSharedPtr
+		ConsoleObjectsMainReference[FoundIndex]->ConsoleObjectPtr = nullptr;
+
+		ConsoleObjectsMainReference.RemoveAt(FoundIndex, 1, /*shrink*/false);
 	}
 
-	if (const TWeakPtr<FConsoleVariablesEditorCommandInfo> CommandInfo =
-		FindCommandInfoByName(CommandName); CommandInfo.IsValid())
+	if (!IsEngineExitRequested() && MainPanel.IsValid())
 	{
-		ConsoleObjectsMainReference.Remove(CommandInfo.Pin());
+		// TODO: Request list refresh, don't force it now as many objects could be unregistered in one frame during reloadconfig
+		MainPanel->RefreshList();
 	}
 }
 
diff --git a/Engine/Plugins/Editor/ConsoleVariablesEditor/Source/ConsoleVariablesEditor/Public/ConsoleVariablesEditorCommandInfo.h b/Engine/Plugins/Editor/ConsoleVariablesEditor/Source/ConsoleVariablesEditor/Public/ConsoleVariablesEditorCommandInfo.h
--- a/Engine/Plugins/Editor/ConsoleVariablesEditor/Source/ConsoleVariablesEditor/Public/ConsoleVariablesEditorCommandInfo.h
+++ b/Engine/Plugins/Editor/ConsoleVariablesEditor/Source/ConsoleVariablesEditor/Public/ConsoleVariablesEditorCommandInfo.h
@@ -9,8 +9,6 @@
 #include "Misc/DateTime.h"
 #include "UObject/ObjectMacros.h"
 
-DECLARE_MULTICAST_DELEGATE_OneParam(FOnDetectConsoleObjectUnregistered, FString)
-
 #define LOCTEXT_NAMESPACE "ConsoleVariablesEditor"
 
 struct CONSOLEVARIABLESEDITOR_API FConsoleVariablesEditorCommandInfo
@@ -49,7 +47,6 @@ struct CONSOLEVARIABLESEDITOR_API FConsoleVariablesEditorCommandInfo
 	
 	~FConsoleVariablesEditorCommandInfo()
 	{
-		OnDetectConsoleObjectUnregistered.Remove(OnDetectConsoleObjectUnregisteredHandle);
 		if (!IsEngineExitRequested()) // CrashFix as GetConsoleVariablePtr() returns stale deleted IConsoleVariable during shut down
 		{
 			if (IConsoleVariable* AsVariable = GetConsoleVariablePtr())
@@ -137,9 +134,6 @@ struct CONSOLEVARIABLESEDITOR_API FConsoleVariablesEditorCommandInfo
 	bool bSetInCurrentSession = false;
 	/** When variables change, this callback is executed. */
 	FDelegateHandle OnVariableChangedCallbackHandle;
-	/** When commands are unregistered change, this callback is broadcasted. */
-	FOnDetectConsoleObjectUnregistered OnDetectConsoleObjectUnregistered;
-	FDelegateHandle OnDetectConsoleObjectUnregisteredHandle;
 	/** A mapping of SetBy console variable flags to information like the associated display text. */
 	static const inline TArray<FStaticConsoleVariableFlagInfo> SupportedFlags =
 	{
diff --git a/Engine/Plugins/Editor/ConsoleVariablesEditor/Source/ConsoleVariablesEditor/Public/ConsoleVariablesEditorModule.h b/Engine/Plugins/Editor/ConsoleVariablesEditor/Source/ConsoleVariablesEditor/Public/ConsoleVariablesEditorModule.h
--- a/Engine/Plugins/Editor/ConsoleVariablesEditor/Source/ConsoleVariablesEditor/Public/ConsoleVariablesEditorModule.h
+++ b/Engine/Plugins/Editor/ConsoleVariablesEditor/Source/ConsoleVariablesEditor/Public/ConsoleVariablesEditorModule.h
@@ -135,8 +135,7 @@ private:
 	void RegisterProjectSettings() const;
 
 	void OnConsoleVariableChanged(IConsoleVariable* ChangedVariable);
-	/** In the event a console object is unregistered, this failsafe callback will clean up the associated list item and command info object. */
-	void OnDetectConsoleObjectUnregistered(FString CommandName);
+	void OnConsoleObjectUnregistered(const TCHAR* InName, IConsoleObject* InConsoleObject);
 
 	TObjectPtr<UConsoleVariablesAsset> AllocateTransientPreset(const FName DesiredName) const;
 	void CreateEditingPresets();
@@ -164,4 +163,6 @@ private:
 	 * a remote cvar change creating a ping/pong cvar updates effect between nodes.
 	 */
 	TMap<FString, int32> CommandsReceivedFromMultiUser;
+
+	FDelegateHandle OnConsoleObjectUnregisteredHandle;
 };
diff --git a/Engine/Source/Runtime/Core/Private/HAL/ConsoleManager.cpp b/Engine/Source/Runtime/Core/Private/HAL/ConsoleManager.cpp
--- a/Engine/Source/Runtime/Core/Private/HAL/ConsoleManager.cpp
+++ b/Engine/Source/Runtime/Core/Private/HAL/ConsoleManager.cpp
@@ -1967,6 +1967,8 @@ void FConsoleManager::UnregisterConsoleObject(const TCHAR* Name, bool bKeepState
 
 	if(Object)
 	{
+		ConsoleObjectUnregisteredDelegate.Broadcast(Name, Object);
+
 		IConsoleVariable* CVar = Object->AsVariable();
 
 		if (CVar)
@@ -2830,6 +2832,11 @@ FConsoleVariableMulticastDelegate& FConsoleManager::OnCVarUnregistered()
 	return ConsoleVariableUnregisteredDelegate;
 }
 
+FConsoleObjectWithNameMulticastDelegate& FConsoleManager::OnConsoleObjectUnregistered()
+{
+	return ConsoleObjectUnregisteredDelegate;
+}
+
 void FConsoleManager::UnsetAllConsoleVariablesWithTag(FName Tag)
 {
 #if UE_ALLOW_CVAR_HISTORY
diff --git a/Engine/Source/Runtime/Core/Public/HAL/ConsoleManager.h b/Engine/Source/Runtime/Core/Public/HAL/ConsoleManager.h
--- a/Engine/Source/Runtime/Core/Public/HAL/ConsoleManager.h
+++ b/Engine/Source/Runtime/Core/Public/HAL/ConsoleManager.h
@@ -48,6 +48,7 @@ public:
 	CORE_API void OnCVarChanged();
 
 	CORE_API virtual FConsoleVariableMulticastDelegate& OnCVarUnregistered()override;
+	CORE_API virtual FConsoleObjectWithNameMulticastDelegate& OnConsoleObjectUnregistered() override;
 
 	// interface IConsoleManager -----------------------------------
 
@@ -99,6 +100,7 @@ private: // ----------------------------------------------------
 	TArray<FConsoleCommandDelegate>	ConsoleVariableChangeSinks;
 
 	FConsoleVariableMulticastDelegate ConsoleVariableUnregisteredDelegate;
+	FConsoleObjectWithNameMulticastDelegate ConsoleObjectUnregisteredDelegate;
 
 	IConsoleThreadPropagation* ThreadPropagationCallback;
 
diff --git a/Engine/Source/Runtime/Core/Public/HAL/IConsoleManager.h b/Engine/Source/Runtime/Core/Public/HAL/IConsoleManager.h
--- a/Engine/Source/Runtime/Core/Public/HAL/IConsoleManager.h
+++ b/Engine/Source/Runtime/Core/Public/HAL/IConsoleManager.h
@@ -199,6 +199,9 @@ DECLARE_DELEGATE_OneParam( FConsoleVariableDelegate, IConsoleVariable* );
 /** Console variable multicast delegate type. */
 DECLARE_MULTICAST_DELEGATE_OneParam(FConsoleVariableMulticastDelegate, IConsoleVariable*);
 
+/** Console object with name multicast delegate type. */
+DECLARE_MULTICAST_DELEGATE_TwoParams(FConsoleObjectWithNameMulticastDelegate, const TCHAR*, IConsoleObject*);
+
 /** Console command delegate type (takes no arguments.)  This is a void callback function. */
 DECLARE_DELEGATE( FConsoleCommandDelegate );
 
@@ -1087,6 +1090,7 @@ struct IConsoleManager
 	virtual void UnsetAllConsoleVariablesWithTag(FName Tag) = 0;
 	
 	virtual FConsoleVariableMulticastDelegate& OnCVarUnregistered() = 0;
+	virtual FConsoleObjectWithNameMulticastDelegate& OnConsoleObjectUnregistered() = 0;
 
 protected:
 	virtual ~IConsoleManager() { }
