Optimize FConsoleVariablesEditorModule, 500 ms of startup time, 500ms each reloadconfig, 4-10ms each time a cvar changes
#rnx

diff --git a/Engine/Plugins/Editor/ConsoleVariablesEditor/Source/ConsoleVariablesEditor/Private/ConsoleVariablesEditorCommandInfo.cpp b/Engine/Plugins/Editor/ConsoleVariablesEditor/Source/ConsoleVariablesEditor/Private/ConsoleVariablesEditorCommandInfo.cpp
--- a/Engine/Plugins/Editor/ConsoleVariablesEditor/Source/ConsoleVariablesEditor/Private/ConsoleVariablesEditorCommandInfo.cpp
+++ b/Engine/Plugins/Editor/ConsoleVariablesEditor/Source/ConsoleVariablesEditor/Private/ConsoleVariablesEditorCommandInfo.cpp
@@ -54,21 +54,25 @@ void FConsoleVariablesEditorCommandInfo::PrintCommandOrVariable()
 /** Get a reference to the cached console object. May return nullptr if unregistered. */
 IConsoleObject* FConsoleVariablesEditorCommandInfo::GetConsoleObjectPtr()
 {
-	// If the console object ptr goes stale or is older than the specified threshold, try to refresh it
+	// If the console object ptr null, try to find if not already attempted
 	// May return nullptr if unregistered
-	if (!ConsoleObjectPtr ||
-		(FDateTime::UtcNow() - TimeOfLastConsoleObjectRefresh).GetTotalSeconds() > ConsoleObjectRefreshThreshold)
+	if (!ConsoleObjectPtr && !bHasAttemptedFind)
 	{
-		FString CommandKey = Command; 
+		bHasAttemptedFind = true;
+
 		// Remove additional params, if they exist
-		const int32 IndexOfSpace = CommandKey.Find(" ");
-		if (IndexOfSpace != INDEX_NONE)
+		int32 IndexOfSpace = INDEX_NONE;
+		if (Command.FindChar(TEXT(' '), IndexOfSpace))
 		{
-			CommandKey = CommandKey.Left(IndexOfSpace).TrimStartAndEnd();
+			const FStringView CommandKeyView = FStringView(*Command, IndexOfSpace).TrimStartAndEnd();
+			TStringBuilder<260> CommandKey;
+			CommandKey.Append(CommandKeyView);
+			ConsoleObjectPtr = IConsoleManager::Get().FindConsoleObject(*CommandKey);
+		}
+		else
+		{
+			ConsoleObjectPtr = IConsoleManager::Get().FindConsoleObject(*Command);
 		}
-		
-		ConsoleObjectPtr = IConsoleManager::Get().FindConsoleObject(*CommandKey);
-		TimeOfLastConsoleObjectRefresh = FDateTime::UtcNow();
 	}
 	return ConsoleObjectPtr;
 }
@@ -227,4 +231,9 @@ bool FConsoleVariablesEditorCommandInfo::IsCurrentValueDifferentFromInputValue(c
 	return false;
 }
 
+void FConsoleVariablesEditorCommandInfo::OnConsoleVariableChanged(IConsoleVariable* ChangedVariable)
+{
+	FConsoleVariablesEditorModule::Get().OnConsoleVariableChanged(*this, ChangedVariable);
+}
+
 #undef LOCTEXT_NAMESPACE
diff --git a/Engine/Plugins/Editor/ConsoleVariablesEditor/Source/ConsoleVariablesEditor/Private/ConsoleVariablesEditorModule.cpp b/Engine/Plugins/Editor/ConsoleVariablesEditor/Source/ConsoleVariablesEditor/Private/ConsoleVariablesEditorModule.cpp
--- a/Engine/Plugins/Editor/ConsoleVariablesEditor/Source/ConsoleVariablesEditor/Private/ConsoleVariablesEditorModule.cpp
+++ b/Engine/Plugins/Editor/ConsoleVariablesEditor/Source/ConsoleVariablesEditor/Private/ConsoleVariablesEditorModule.cpp
@@ -48,7 +48,8 @@ const FName FConsoleVariablesEditorModule::ConsoleVariablesToolkitPanelTabId(TEX
 
 FConsoleVariablesEditorModule& FConsoleVariablesEditorModule::Get()
 {
-	return FModuleManager::GetModuleChecked<FConsoleVariablesEditorModule>("ConsoleVariablesEditor");
+	static const FName ConsoleVariablesEditorModuleName("ConsoleVariablesEditor");
+	return FModuleManager::GetModuleChecked<FConsoleVariablesEditorModule>(ConsoleVariablesEditorModuleName);
 }
 
 void FConsoleVariablesEditorModule::StartupModule()
@@ -159,11 +160,13 @@ void FConsoleVariablesEditorModule::QueryAndBeginTrackingConsoleVariables()
 				MakeShared<FConsoleVariablesEditorCommandInfo>(Key);
 			
 			Info->StartupSource = Info->GetSource();
+			Info->ConsoleObjectPtr = ConsoleObject;
+			Info->bHasAttemptedFind = true;
 
 			if (IConsoleVariable* AsVariable = ConsoleObject->AsVariable())
 			{
-				Info->OnVariableChangedCallbackHandle = AsVariable->OnChangedDelegate().AddRaw(
-					this, &FConsoleVariablesEditorModule::OnConsoleVariableChanged);
+				Info->OnVariableChangedCallbackHandle = AsVariable->OnChangedDelegate().AddSP(
+					Info, &FConsoleVariablesEditorCommandInfo::OnConsoleVariableChanged);
 			}
 			
 			AddConsoleObjectCommandInfoToMainReference(Info);
@@ -173,15 +176,14 @@ void FConsoleVariablesEditorModule::QueryAndBeginTrackingConsoleVariables()
 	OnConsoleObjectUnregisteredHandle = IConsoleManager::Get().OnConsoleObjectUnregistered().AddRaw(this, &FConsoleVariablesEditorModule::OnConsoleObjectUnregistered);
 }
 
-TWeakPtr<FConsoleVariablesEditorCommandInfo> FConsoleVariablesEditorModule::FindCommandInfoByName(const FString& NameToSearch, ESearchCase::Type InSearchCase)
+void FConsoleVariablesEditorModule::AddConsoleObjectCommandInfoToMainReference(TSharedRef<FConsoleVariablesEditorCommandInfo> InCommandInfo)
 {
-	TSharedPtr<FConsoleVariablesEditorCommandInfo>* Match = Algo::FindByPredicate(
-		ConsoleObjectsMainReference,
-		[&NameToSearch, InSearchCase](const TSharedPtr<FConsoleVariablesEditorCommandInfo> Comparator)
-		{
-			return Comparator->Command.Equals(NameToSearch, InSearchCase);
-		});
+	ConsoleObjectsMainReference.Add(InCommandInfo->Command, InCommandInfo);
+}
 
+TWeakPtr<FConsoleVariablesEditorCommandInfo> FConsoleVariablesEditorModule::FindCommandInfoByName(const FString& NameToSearch, ESearchCase::Type InSearchCase)
+{
+	TSharedPtr<FConsoleVariablesEditorCommandInfo>* Match = ConsoleObjectsMainReference.Find(NameToSearch);
 	return Match ? *Match : nullptr;
 }
 
@@ -190,8 +192,9 @@ TArray<TWeakPtr<FConsoleVariablesEditorCommandInfo>> FConsoleVariablesEditorModu
 {
 	TArray<TWeakPtr<FConsoleVariablesEditorCommandInfo>> ReturnValue;
 	
-	for (const TSharedPtr<FConsoleVariablesEditorCommandInfo>& CommandInfo : ConsoleObjectsMainReference)
+	for (const TPair<FString, TSharedPtr<FConsoleVariablesEditorCommandInfo>>& It : ConsoleObjectsMainReference)
 	{
+		const TSharedPtr<FConsoleVariablesEditorCommandInfo>& CommandInfo = It.Value;
 		FString CommandSearchableText = CommandInfo->Command + " " + CommandInfo->GetHelpText() + " " + CommandInfo->GetSourceAsText().ToString();
 		// Match any
 		for (const FString& Token : InTokens)
@@ -223,19 +226,6 @@ TArray<TWeakPtr<FConsoleVariablesEditorCommandInfo>> FConsoleVariablesEditorModu
 	return ReturnValue;
 }
 
-TWeakPtr<FConsoleVariablesEditorCommandInfo> FConsoleVariablesEditorModule::FindCommandInfoByConsoleObjectReference(
-	IConsoleObject* InConsoleObjectReference)
-{
-	TSharedPtr<FConsoleVariablesEditorCommandInfo>* Match = Algo::FindByPredicate(
-	ConsoleObjectsMainReference,
-	[InConsoleObjectReference](const TSharedPtr<FConsoleVariablesEditorCommandInfo> Comparator)
-	{
-		return Comparator->GetConsoleObjectPtr() == InConsoleObjectReference;
-	});
-
-	return Match ? *Match : nullptr;
-}
-
 TObjectPtr<UConsoleVariablesAsset> FConsoleVariablesEditorModule::GetPresetAsset() const
 {
 	return EditingPresetAsset;
@@ -363,7 +353,7 @@ void FConsoleVariablesEditorModule::RegisterProjectSettings() const
 	}
 }
 
-void FConsoleVariablesEditorModule::OnConsoleVariableChanged(IConsoleVariable* ChangedVariable)
+void FConsoleVariablesEditorModule::OnConsoleVariableChanged(FConsoleVariablesEditorCommandInfo& CommandInfo, IConsoleVariable* ChangedVariable)
 {
 	// Note: We disable tracking during automation tests to prevent FindConsoleObject() related performance warnings.
 	if (IsEngineExitRequested() || GIsAutomationTesting)
@@ -373,11 +363,8 @@ void FConsoleVariablesEditorModule::OnConsoleVariableChanged(IConsoleVariable* C
 
 	check(EditingPresetAsset);
 
-	if (const TWeakPtr<FConsoleVariablesEditorCommandInfo> CommandInfo =
-		FindCommandInfoByConsoleObjectReference(ChangedVariable); CommandInfo.IsValid())
 	{
-		const TSharedPtr<FConsoleVariablesEditorCommandInfo>& PinnedCommand = CommandInfo.Pin();
-		const FString& Key = PinnedCommand->Command;
+		const FString& Key = CommandInfo.Command;
 
 		FConsoleVariablesEditorAssetSaveData FoundData;
 		bool bIsVariableCurrentlyTracked = EditingPresetAsset->FindSavedDataByCommandString(Key, FoundData, ESearchCase::IgnoreCase);
@@ -391,7 +378,7 @@ void FConsoleVariablesEditorModule::OnConsoleVariableChanged(IConsoleVariable* C
 			// Check if the changed value differs from the startup value before tracking it
 			if (Settings->bAddAllChangedConsoleVariablesToCurrentPreset &&
 				!Settings->ChangedConsoleVariableSkipList.Contains(Key) &&
-				PinnedCommand->IsCurrentValueDifferentFromInputValue(PinnedCommand->StartupValueAsString))
+				CommandInfo.IsCurrentValueDifferentFromInputValue(CommandInfo.StartupValueAsString))
 			{
 				if (MainPanel.IsValid())
 				{
@@ -447,37 +434,20 @@ void FConsoleVariablesEditorModule::OnConsoleObjectUnregistered(const TCHAR* InN
 		}
 	}
 
-	int32 FoundIndex = INDEX_NONE;
-	for (int32 i = 0; i < ConsoleObjectsMainReference.Num(); ++i)
+	TSharedPtr<FConsoleVariablesEditorCommandInfo> Found;
+	if (ConsoleObjectsMainReference.RemoveAndCopyValueByHash(GetTypeHash(FStringView(InName)), InName, Found))
 	{
-		if (ConsoleObjectsMainReference[i])
+		if (Found.IsValid())
 		{
-			if (ConsoleObjectsMainReference[i]->ConsoleObjectPtr == InConsoleObject)
-			{
-				FoundIndex = i;
-				break;
-			}
-			// GetConsoleObjectPtr() is overwriting ConsoleObjectPtr so we need to search by command name as well
-			else if (InName && ConsoleObjectsMainReference[i]->Command == InName)
-			{
-				FoundIndex = i;
-				break;
-			}
+			// Null out deleted ConsoleObjectPtr for anything still holding reference to TSharedPtr
+			Found->ConsoleObjectPtr = nullptr;
 		}
-	}
-
-	if (FoundIndex != -1)
-	{
-		// Null out deleted ConsoleObjectPtr for anything still holding reference to TSharedPtr
-		ConsoleObjectsMainReference[FoundIndex]->ConsoleObjectPtr = nullptr;
 
-		ConsoleObjectsMainReference.RemoveAt(FoundIndex, 1, /*shrink*/false);
-	}
-
-	if (!IsEngineExitRequested() && MainPanel.IsValid())
-	{
-		// TODO: Request list refresh, don't force it now as many objects could be unregistered in one frame during reloadconfig
-		MainPanel->RefreshList();
+		if (!IsEngineExitRequested() && MainPanel.IsValid())
+		{
+			// TODO: Request list refresh, don't force it now as many objects could be unregistered in one frame during reloadconfig
+			MainPanel->RefreshList();
+		}
 	}
 }
 
diff --git a/Engine/Plugins/Editor/ConsoleVariablesEditor/Source/ConsoleVariablesEditor/Public/ConsoleVariablesEditorCommandInfo.h b/Engine/Plugins/Editor/ConsoleVariablesEditor/Source/ConsoleVariablesEditor/Public/ConsoleVariablesEditorCommandInfo.h
--- a/Engine/Plugins/Editor/ConsoleVariablesEditor/Source/ConsoleVariablesEditor/Public/ConsoleVariablesEditorCommandInfo.h
+++ b/Engine/Plugins/Editor/ConsoleVariablesEditor/Source/ConsoleVariablesEditor/Public/ConsoleVariablesEditorCommandInfo.h
@@ -116,14 +116,14 @@ struct CONSOLEVARIABLESEDITOR_API FConsoleVariablesEditorCommandInfo
  
 	bool IsCurrentValueDifferentFromInputValue(const FString& InValueToCompare);
  
+	void OnConsoleVariableChanged(IConsoleVariable* ChangedVariable);
+
 	/** The actual string key or name */
 	UPROPERTY()
 	FString Command;
 	EConsoleObjectType ObjectType = EConsoleObjectType::NullObject;
 	/** This object is periodically refreshed to mitigate the occurrence of stale pointers. */
-	IConsoleObject* ConsoleObjectPtr;
-	FDateTime TimeOfLastConsoleObjectRefresh;
-	double ConsoleObjectRefreshThreshold = 1.0;
+	IConsoleObject* ConsoleObjectPtr = nullptr;
 	
 	/** The value of this variable (if Variable object type) when the module started in this session after it may have been set by an ini file. */
 	FString StartupValueAsString;
@@ -132,6 +132,7 @@ struct CONSOLEVARIABLESEDITOR_API FConsoleVariablesEditorCommandInfo
 	EConsoleVariableFlags StartupSource = ECVF_Default;
 	/** If the variable was last changed by the current preset */
 	bool bSetInCurrentSession = false;
+	bool bHasAttemptedFind = false;
 	/** When variables change, this callback is executed. */
 	FDelegateHandle OnVariableChangedCallbackHandle;
 	/** A mapping of SetBy console variable flags to information like the associated display text. */
diff --git a/Engine/Plugins/Editor/ConsoleVariablesEditor/Source/ConsoleVariablesEditor/Public/ConsoleVariablesEditorModule.h b/Engine/Plugins/Editor/ConsoleVariablesEditor/Source/ConsoleVariablesEditor/Public/ConsoleVariablesEditorModule.h
--- a/Engine/Plugins/Editor/ConsoleVariablesEditor/Source/ConsoleVariablesEditor/Public/ConsoleVariablesEditorModule.h
+++ b/Engine/Plugins/Editor/ConsoleVariablesEditor/Source/ConsoleVariablesEditor/Public/ConsoleVariablesEditorModule.h
@@ -38,10 +38,7 @@ public:
 	/** Find all console variables and cache their startup values */
 	void QueryAndBeginTrackingConsoleVariables();
 
-	void AddConsoleObjectCommandInfoToMainReference(TSharedRef<FConsoleVariablesEditorCommandInfo> InCommandInfo)
-	{
-		ConsoleObjectsMainReference.Add(InCommandInfo);
-	}
+	void AddConsoleObjectCommandInfoToMainReference(TSharedRef<FConsoleVariablesEditorCommandInfo> InCommandInfo);
 
 	/** Find a tracked console variable by the command string with optional case sensitivity. */
 	CONSOLEVARIABLESEDITOR_API TWeakPtr<FConsoleVariablesEditorCommandInfo> FindCommandInfoByName(
@@ -55,14 +52,6 @@ public:
 	CONSOLEVARIABLESEDITOR_API TArray<TWeakPtr<FConsoleVariablesEditorCommandInfo>> FindCommandInfosMatchingTokens(
 		const TArray<FString>& InTokens, ESearchCase::Type InSearchCase = ESearchCase::IgnoreCase);
 
-	/**
-	 *Find a tracked console variable by its console object reference.
-	 *Note that some commands do not have an associated console object (such as 'stat unit') and will not be found with this method.
-	 *It's normally safer to use FindCommandInfoByName() instead.
-	 */
-	CONSOLEVARIABLESEDITOR_API TWeakPtr<FConsoleVariablesEditorCommandInfo> FindCommandInfoByConsoleObjectReference(
-		IConsoleObject* InConsoleObjectReference);
-	
 	[[nodiscard]] TObjectPtr<UConsoleVariablesAsset> GetPresetAsset() const;
 	[[nodiscard]] TObjectPtr<UConsoleVariablesAsset> GetGlobalSearchAsset() const;
 	
@@ -90,6 +79,8 @@ public:
 	virtual void AddReferencedObjects( FReferenceCollector& Collector )  override;
 	virtual FString GetReferencerName() const override;
 
+	void OnConsoleVariableChanged(FConsoleVariablesEditorCommandInfo& CommandInfo, IConsoleVariable* ChangedVariable);
+
 private:
 	/** Scoped based struct to track inbound cvars set by Multi-user. */
 	struct FScopeMultiUserReceiveCVar
@@ -133,8 +124,6 @@ private:
 
 	void RegisterMenuItem();
 	void RegisterProjectSettings() const;
-
-	void OnConsoleVariableChanged(IConsoleVariable* ChangedVariable);
 	void OnConsoleObjectUnregistered(const TCHAR* InName, IConsoleObject* InConsoleObject);
 
 	TObjectPtr<UConsoleVariablesAsset> AllocateTransientPreset(const FName DesiredName) const;
@@ -155,7 +144,7 @@ private:
 	TObjectPtr<UConsoleVariablesAsset> EditingGlobalSearchAsset = nullptr;
 
 	/** All tracked variables and their default, startup, and current values */
-	TArray<TSharedPtr<FConsoleVariablesEditorCommandInfo>> ConsoleObjectsMainReference;
+	TMap<FString, TSharedPtr<FConsoleVariablesEditorCommandInfo>> ConsoleObjectsMainReference;
 
 	/**
 	 * A map of commands invoked by Multi-user.  Each command has a corresponding reference count.
diff --git a/Engine/Plugins/Editor/ConsoleVariablesEditor/Source/ConsoleVariablesEditorRuntime/Private/ConsoleVariablesAsset.cpp b/Engine/Plugins/Editor/ConsoleVariablesEditor/Source/ConsoleVariablesEditorRuntime/Private/ConsoleVariablesAsset.cpp
--- a/Engine/Plugins/Editor/ConsoleVariablesEditor/Source/ConsoleVariablesEditorRuntime/Private/ConsoleVariablesAsset.cpp
+++ b/Engine/Plugins/Editor/ConsoleVariablesEditor/Source/ConsoleVariablesEditorRuntime/Private/ConsoleVariablesAsset.cpp
@@ -86,14 +86,18 @@ void UConsoleVariablesAsset::ReplaceSavedCommands(const TArray<FConsoleVariables
 }
 
 bool UConsoleVariablesAsset::FindSavedDataByCommandString(
-	const FString InCommandString, FConsoleVariablesEditorAssetSaveData& OutValue, const ESearchCase::Type SearchCase) const
+	const FString& InCommandString, FConsoleVariablesEditorAssetSaveData& OutValue, const ESearchCase::Type SearchCase) const
 {
-	for (const FConsoleVariablesEditorAssetSaveData& Command : SavedCommands)
+	if (!SavedCommands.IsEmpty())
 	{
-		if (Command.CommandName.TrimStartAndEnd().Equals(InCommandString.TrimStartAndEnd(), SearchCase))
+		const FStringView InCommandTrimmed = FStringView(InCommandString).TrimStartAndEnd();
+		for (const FConsoleVariablesEditorAssetSaveData& Command : SavedCommands)
 		{
-			OutValue = Command;
-			return true;
+			if (FStringView(Command.CommandName).TrimStartAndEnd().Equals(InCommandTrimmed, SearchCase))
+			{
+				OutValue = Command;
+				return true;
+			}
 		}
 	}
 	
@@ -115,7 +119,7 @@ void UConsoleVariablesAsset::AddOrSetConsoleObjectSavedData(const FConsoleVariab
 	}
 }
 
-bool UConsoleVariablesAsset::RemoveConsoleVariable(const FString InCommandString)
+bool UConsoleVariablesAsset::RemoveConsoleVariable(const FString& InCommandString)
 {
 	FConsoleVariablesEditorAssetSaveData ExistingData;
 	int32 RemoveCount = 0;
diff --git a/Engine/Plugins/Editor/ConsoleVariablesEditor/Source/ConsoleVariablesEditorRuntime/Public/ConsoleVariablesAsset.h b/Engine/Plugins/Editor/ConsoleVariablesEditor/Source/ConsoleVariablesEditorRuntime/Public/ConsoleVariablesAsset.h
--- a/Engine/Plugins/Editor/ConsoleVariablesEditor/Source/ConsoleVariablesEditorRuntime/Public/ConsoleVariablesAsset.h
+++ b/Engine/Plugins/Editor/ConsoleVariablesEditor/Source/ConsoleVariablesEditorRuntime/Public/ConsoleVariablesAsset.h
@@ -86,7 +86,7 @@ public:
 	/** Outputs the FConsoleVariablesEditorAssetSaveData matching InCommand. Returns whether a match was found. Case sensitive. */
 	UFUNCTION(BlueprintCallable, Category="Console Variables Asset")
 	bool FindSavedDataByCommandString(
-		const FString InCommandString, FConsoleVariablesEditorAssetSaveData& OutValue, const ESearchCase::Type SearchCase = ESearchCase::IgnoreCase) const;
+		const FString& InCommandString, FConsoleVariablesEditorAssetSaveData& OutValue, const ESearchCase::Type SearchCase = ESearchCase::IgnoreCase) const;
 
 	/** Set the value of a saved console variable if the name matches; add a new console variable to the list if a match is not found. */
 	UFUNCTION(BlueprintCallable, Category="Console Variables Asset")
@@ -94,7 +94,7 @@ public:
 
 	/** Returns true if the element was found and successfully removed. */
 	UFUNCTION(BlueprintCallable, Category="Console Variables Asset")
-	bool RemoveConsoleVariable(const FString InCommandString);
+	bool RemoveConsoleVariable(const FString& InCommandString);
 
 	/** Copy data from input asset to this asset */
 	UFUNCTION(BlueprintCallable, Category="Console Variables Asset")
