Compile deprecation warning fix
#rnx

diff --git a/Engine/Plugins/Editor/ConsoleVariablesEditor/Source/ConsoleVariablesEditor/Private/ConsoleVariablesEditorModule.cpp b/Engine/Plugins/Editor/ConsoleVariablesEditor/Source/ConsoleVariablesEditor/Private/ConsoleVariablesEditorModule.cpp
--- a/Engine/Plugins/Editor/ConsoleVariablesEditor/Source/ConsoleVariablesEditor/Private/ConsoleVariablesEditorModule.cpp
+++ b/Engine/Plugins/Editor/ConsoleVariablesEditor/Source/ConsoleVariablesEditor/Private/ConsoleVariablesEditorModule.cpp
@@ -435,7 +435,7 @@ void FConsoleVariablesEditorModule::OnConsoleObjectUnregistered(const TCHAR* InN
 	}
 
 	TSharedPtr<FConsoleVariablesEditorCommandInfo> Found;
-	if (ConsoleObjectsMainReference.RemoveAndCopyValueByHash(GetTypeHash(FStringView(InName)), InName, Found))
+	if (ConsoleObjectsMainReference.RemoveAndCopyValueByHash(FCrc::Strihash_DEPRECATED(InName), InName, Found))
 	{
 		if (Found.IsValid())
 		{
