Fix ConsoleVariablesEditor crash on editor exit
#rnx

diff --git a/Engine/Plugins/Editor/ConsoleVariablesEditor/Source/ConsoleVariablesEditor/Private/ConsoleVariablesEditorModule.cpp b/Engine/Plugins/Editor/ConsoleVariablesEditor/Source/ConsoleVariablesEditor/Private/ConsoleVariablesEditorModule.cpp
--- a/Engine/Plugins/Editor/ConsoleVariablesEditor/Source/ConsoleVariablesEditor/Private/ConsoleVariablesEditorModule.cpp
+++ b/Engine/Plugins/Editor/ConsoleVariablesEditor/Source/ConsoleVariablesEditor/Private/ConsoleVariablesEditorModule.cpp
@@ -63,17 +63,20 @@ void FConsoleVariablesEditorModule::StartupModule()
 	FCoreDelegates::OnFEngineLoopInitComplete.AddRaw(this, &FConsoleVariablesEditorModule::OnFEngineLoopInitComplete);
 }
 
-void FConsoleVariablesEditorModule::ShutdownModule()
+void FConsoleVariablesEditorModule::PreUnloadCallback()
 {
-	UToolMenus::UnregisterOwner(this);
-
-	FCoreDelegates::OnFEngineLoopInitComplete.RemoveAll(this);
-
 	if (OnConsoleObjectUnregisteredHandle.IsValid())
 	{
 		IConsoleManager::Get().OnConsoleObjectUnregistered().Remove(OnConsoleObjectUnregisteredHandle);
 		OnConsoleObjectUnregisteredHandle.Reset();
 	}
+}
+
+void FConsoleVariablesEditorModule::ShutdownModule()
+{
+	UToolMenus::UnregisterOwner(this);
+
+	FCoreDelegates::OnFEngineLoopInitComplete.RemoveAll(this);
 
 	FConsoleVariablesEditorStyle::Shutdown();
 	
diff --git a/Engine/Plugins/Editor/ConsoleVariablesEditor/Source/ConsoleVariablesEditor/Public/ConsoleVariablesEditorModule.h b/Engine/Plugins/Editor/ConsoleVariablesEditor/Source/ConsoleVariablesEditor/Public/ConsoleVariablesEditorModule.h
--- a/Engine/Plugins/Editor/ConsoleVariablesEditor/Source/ConsoleVariablesEditor/Public/ConsoleVariablesEditorModule.h
+++ b/Engine/Plugins/Editor/ConsoleVariablesEditor/Source/ConsoleVariablesEditor/Public/ConsoleVariablesEditorModule.h
@@ -25,6 +25,7 @@ public:
 
 	//~ Begin IModuleInterface Interface
 	virtual void StartupModule() override;
+	virtual void PreUnloadCallback() override;
 	virtual void ShutdownModule() override;
 	//~ End IModuleInterface Interface
 
