Fix a crash on shutdown
#rnx

diff --git a/Engine/Plugins/Editor/ConsoleVariablesEditor/Source/ConsoleVariablesEditor/Private/ConsoleVariablesEditorModule.cpp b/Engine/Plugins/Editor/ConsoleVariablesEditor/Source/ConsoleVariablesEditor/Private/ConsoleVariablesEditorModule.cpp
--- a/Engine/Plugins/Editor/ConsoleVariablesEditor/Source/ConsoleVariablesEditor/Private/ConsoleVariablesEditorModule.cpp
+++ b/Engine/Plugins/Editor/ConsoleVariablesEditor/Source/ConsoleVariablesEditor/Private/ConsoleVariablesEditorModule.cpp
@@ -429,6 +429,12 @@ void FConsoleVariablesEditorModule::OnConsoleVariableChanged(FConsoleVariablesEd
 
 void FConsoleVariablesEditorModule::OnConsoleObjectUnregistered(const TCHAR* InName, IConsoleObject* InConsoleObject)
 {
+	// Note: EditingPresetAsset is not being nulled out when destroyed so accessing EditingPresetAsset on shutdown is crashing
+	if (IsEngineExitRequested())
+	{
+		return;
+	}
+
 	if (ensure(InName))
 	{
 		if (ensure(EditingPresetAsset))
@@ -446,7 +452,7 @@ void FConsoleVariablesEditorModule::OnConsoleObjectUnregistered(const TCHAR* InN
 			Found->ConsoleObjectPtr = nullptr;
 		}
 
-		if (!IsEngineExitRequested() && MainPanel.IsValid())
+		if (MainPanel.IsValid())
 		{
 			// TODO: Request list refresh, don't force it now as many objects could be unregistered in one frame during reloadconfig
 			MainPanel->RefreshList();
