Transmit how a console variable was set so it will be set with the same state when received by multi-user.

diff --git a/Engine/Plugins/Editor/ConsoleVariablesEditor/Source/ConsoleVariablesEditor/Private/ConsoleVariablesEditorCommandInfo.cpp b/Engine/Plugins/Editor/ConsoleVariablesEditor/Source/ConsoleVariablesEditor/Private/ConsoleVariablesEditorCommandInfo.cpp
--- a/Engine/Plugins/Editor/ConsoleVariablesEditor/Source/ConsoleVariablesEditor/Private/ConsoleVariablesEditorCommandInfo.cpp
+++ b/Engine/Plugins/Editor/ConsoleVariablesEditor/Source/ConsoleVariablesEditor/Private/ConsoleVariablesEditorCommandInfo.cpp
@@ -34,7 +34,7 @@ void FConsoleVariablesEditorCommandInfo::ExecuteCommand(
 	}
 	if (!bIsInteractiveChange && bShouldTransactInConcert)
 	{
-		FConsoleVariablesEditorModule::Get().SendMultiUserConsoleVariableChange(ERemoteCVarChangeType::Update, Command, NewValueAsString);
+		FConsoleVariablesEditorModule::Get().SendMultiUserConsoleVariableChange(ERemoteCVarChangeType::Update, Command, NewValueAsString, GetSource());
 	}
 }
 
diff --git a/Engine/Plugins/Editor/ConsoleVariablesEditor/Source/ConsoleVariablesEditor/Private/ConsoleVariablesEditorModule.cpp b/Engine/Plugins/Editor/ConsoleVariablesEditor/Source/ConsoleVariablesEditor/Private/ConsoleVariablesEditorModule.cpp
--- a/Engine/Plugins/Editor/ConsoleVariablesEditor/Source/ConsoleVariablesEditor/Private/ConsoleVariablesEditorModule.cpp
+++ b/Engine/Plugins/Editor/ConsoleVariablesEditor/Source/ConsoleVariablesEditor/Private/ConsoleVariablesEditorModule.cpp
@@ -8,6 +8,7 @@
 #include "ConsoleVariablesEditorLog.h"
 #include "ConsoleVariablesEditorProjectSettings.h"
 #include "ConsoleVariablesEditorStyle.h"
+#include "HAL/IConsoleManager.h"
 #include "MultiUser/ConsoleVariableSync.h"
 #include "MultiUser/ConsoleVariableSyncData.h"
 #include "Views/MainPanel/ConsoleVariablesEditorMainPanel.h"
@@ -280,16 +281,16 @@ bool FConsoleVariablesEditorModule::PopulateGlobalSearchAssetWithVariablesMatchi
 	return EditingGlobalSearchAsset->GetSavedCommandsCount() > 0;
 }
 
-void FConsoleVariablesEditorModule::SendMultiUserConsoleVariableChange(ERemoteCVarChangeType InChangeType, const FString& InVariableName, const FString& InValueAsString)
+void FConsoleVariablesEditorModule::SendMultiUserConsoleVariableChange(ERemoteCVarChangeType InChangeType, const FString& InVariableName, const FString& InValueAsString, EConsoleVariableFlags InFlags)
 {
 	using namespace UE::ConsoleVariablesEditor::ModuleUtils::Private; 
 	if (MainPanel->GetMultiUserManager().IsLocalUserInMultiUserSession() && !CheckIsRunningPieAndWarnSyncDisabledInPie())
 	{
-		MainPanel->GetMultiUserManager().SendConsoleVariableChange(InChangeType, InVariableName, InValueAsString);
+		MainPanel->GetMultiUserManager().SendConsoleVariableChange(InChangeType, InVariableName, InValueAsString, InFlags);
 	}
 }
 
-void FConsoleVariablesEditorModule::OnRemoteCvarChanged(ERemoteCVarChangeType InChangeType, const FString InName, const FString InValue)
+void FConsoleVariablesEditorModule::OnRemoteCvarChanged(ERemoteCVarChangeType InChangeType, const FString InName, const FString InValue, EConsoleVariableFlags InFlags)
 {
 	using namespace UE::ConsoleVariablesEditor::ModuleUtils::Private; 
 	if (MainPanel->GetMultiUserManager().IsLocalUserInMultiUserSession() &&
@@ -298,8 +299,16 @@ void FConsoleVariablesEditorModule::OnRemoteCvarChanged(ERemoteCVarChangeType In
 	{
 		FScopeMultiUserReceiveCVar CVarChange(*InName, CommandsReceivedFromMultiUser);
 
-		GEngine->Exec(FConsoleVariablesEditorCommandInfo::GetCurrentWorld(),
+		FConsoleVariablesEditorCommandInfo Info(InName);
+		if (IConsoleVariable* AsVariable = Info.GetConsoleVariablePtr())
+		{
+			AsVariable->Set(*InValue, InFlags);
+		}
+		else
+		{
+			GEngine->Exec(FConsoleVariablesEditorCommandInfo::GetCurrentWorld(),
 					  *FString::Printf(TEXT("%s %s"), *InName, *InValue));
+		}
 
 		if (InChangeType == ERemoteCVarChangeType::Remove)
 		{
@@ -421,7 +430,7 @@ void FConsoleVariablesEditorModule::OnConsoleVariableChanged(FConsoleVariablesEd
 			 */
 			if (!CommandsReceivedFromMultiUser.Find(Key))
 			{
-				SendMultiUserConsoleVariableChange(ERemoteCVarChangeType::Update, Key, ChangedVariable->GetString());
+				SendMultiUserConsoleVariableChange(ERemoteCVarChangeType::Update, Key, ChangedVariable->GetString(), ChangedVariable->GetFlags());
 			}
 		}
 	}
diff --git a/Engine/Plugins/Editor/ConsoleVariablesEditor/Source/ConsoleVariablesEditor/Private/MultiUser/ConsoleVariableSync.cpp b/Engine/Plugins/Editor/ConsoleVariablesEditor/Source/ConsoleVariablesEditor/Private/MultiUser/ConsoleVariableSync.cpp
--- a/Engine/Plugins/Editor/ConsoleVariablesEditor/Source/ConsoleVariablesEditor/Private/MultiUser/ConsoleVariableSync.cpp
+++ b/Engine/Plugins/Editor/ConsoleVariablesEditor/Source/ConsoleVariablesEditor/Private/MultiUser/ConsoleVariableSync.cpp
@@ -4,6 +4,8 @@
 
 #include "ConcertTransportMessages.h"
 #include "Delegates/IDelegateInstance.h"
+#include "HAL/IConsoleManager.h"
+
 #include "IConcertClient.h"
 #include "IConcertSyncClient.h"
 #include "IConcertSession.h"
@@ -328,11 +330,22 @@ struct FManagerImpl
 			return ERemoteCVarChangeType::Update;
 		}
 	}
+
 	void HandleCVarSet(const FConcertSessionContext& Context, const FConcertSetConsoleVariableEvent& InEvent)
 	{
+		// Ignore this message if we generated it
+		if (TSharedPtr<IConcertClientSession> Session = WeakSession.Pin())
+		{
+			if (Context.SourceEndpointId == Session->GetSessionClientEndpointId())
+			{
+				return;
+			}
+		}
+
 		if (CanReceiveConsoleVariable())
 		{
-			RemoteCVarChanged.Broadcast(ConvertChangeType(InEvent.ChangeType), InEvent.Variable, InEvent.Value);
+			EConsoleVariableFlags Flags = static_cast<EConsoleVariableFlags>(InEvent.Flags);
+			RemoteCVarChanged.Broadcast(ConvertChangeType(InEvent.ChangeType), InEvent.Variable, InEvent.Value, Flags);
 		}
 	}
 
@@ -415,7 +428,7 @@ struct FManagerImpl
 	 * Sends the given console variable name with the specified value to all connected endpoints.  It is up to the
 	 * endpoint to implement the change locally based on configured synchronization state.
 	 */
-	void SendConsoleVariableChange(ERemoteCVarChangeType InChangeType, FString InName, FString InValue)
+	void SendConsoleVariableChange(ERemoteCVarChangeType InChangeType, FString InName, FString InValue, EConsoleVariableFlags InFlags)
 	{
 		TSharedPtr<IConcertClientSession> Session = WeakSession.Pin();
 		if (!bIsEnabled || !Session.IsValid())
@@ -423,7 +436,8 @@ struct FManagerImpl
 			return;
 		}
 
-		FConcertSetConsoleVariableEvent OutEvent{ConvertChangeType(InChangeType), MoveTemp(InName), MoveTemp(InValue)};
+		uint32 OutFlags = static_cast<uint32>(InFlags);
+		FConcertSetConsoleVariableEvent OutEvent{ConvertChangeType(InChangeType), MoveTemp(InName), MoveTemp(InValue), InFlags};
 		Session->SendCustomEvent(OutEvent, Session->GetSessionClientEndpointIds(),
 								 EConcertMessageFlags::ReliableOrdered | EConcertMessageFlags::UniqueId);
 	}
@@ -495,9 +509,9 @@ FOnMultiUserConnectionChange& FManager::OnConnectionChange()
 	return Implementation->ConnectionChanged;
 }
 
-void FManager::SendConsoleVariableChange(ERemoteCVarChangeType InChangeType, FString InName, FString InValue)
+void FManager::SendConsoleVariableChange(ERemoteCVarChangeType InChangeType, FString InName, FString InValue, EConsoleVariableFlags InFlags)
 {
-	Implementation->SendConsoleVariableChange(InChangeType, MoveTemp(InName), MoveTemp(InValue));
+	Implementation->SendConsoleVariableChange(InChangeType, MoveTemp(InName), MoveTemp(InValue), InFlags);
 }
 
 void FManager::SendListItemCheckStateChange(FString InName, ECheckBoxState InCheckedState)
diff --git a/Engine/Plugins/Editor/ConsoleVariablesEditor/Source/ConsoleVariablesEditor/Private/MultiUser/ConsoleVariableSync.h b/Engine/Plugins/Editor/ConsoleVariablesEditor/Source/ConsoleVariablesEditor/Private/MultiUser/ConsoleVariableSync.h
--- a/Engine/Plugins/Editor/ConsoleVariablesEditor/Source/ConsoleVariablesEditor/Private/MultiUser/ConsoleVariableSync.h
+++ b/Engine/Plugins/Editor/ConsoleVariablesEditor/Source/ConsoleVariablesEditor/Private/MultiUser/ConsoleVariableSync.h
@@ -5,9 +5,11 @@
 #include "Containers/UnrealString.h"
 #include "Delegates/Delegate.h"
 #include "Delegates/DelegateCombinations.h"
+#include "HAL/IConsoleManager.h"
 #include "Styling/SlateTypes.h"
 #include "Templates/UniquePtr.h"
 
+
 #include "ConcertMessages.h"
 
 enum class ERemoteCVarChangeType : uint8
@@ -17,7 +19,7 @@ enum class ERemoteCVarChangeType : uint8
 	Remove
 };
 
-DECLARE_MULTICAST_DELEGATE_ThreeParams(FOnRemoteCVarChange, ERemoteCVarChangeType, FString, FString);
+DECLARE_MULTICAST_DELEGATE_FourParams(FOnRemoteCVarChange, ERemoteCVarChangeType, FString, FString, EConsoleVariableFlags);
 DECLARE_MULTICAST_DELEGATE_TwoParams(FOnRemoteListItemCheckStateChange, FString, ECheckBoxState);
 DECLARE_MULTICAST_DELEGATE_OneParam(FOnMultiUserConnectionChange, EConcertConnectionStatus);
 
@@ -47,7 +49,7 @@ struct FManager
 	FOnMultiUserConnectionChange& OnConnectionChange();
 
 	/** Sends the named console variable with value to all connected Multi-user clients */
-	void SendConsoleVariableChange(ERemoteCVarChangeType InChangeType, FString InName, FString InValue);
+	void SendConsoleVariableChange(ERemoteCVarChangeType InChangeType, FString InName, FString InValue, EConsoleVariableFlags InFlags);
 
 	/** Sends the named console variable with value to all connected Multi-user clients */
 	void SendListItemCheckStateChange(FString InName, ECheckBoxState InCheckedState);
diff --git a/Engine/Plugins/Editor/ConsoleVariablesEditor/Source/ConsoleVariablesEditor/Private/MultiUser/ConsoleVariableSyncData.h b/Engine/Plugins/Editor/ConsoleVariablesEditor/Source/ConsoleVariablesEditor/Private/MultiUser/ConsoleVariableSyncData.h
--- a/Engine/Plugins/Editor/ConsoleVariablesEditor/Source/ConsoleVariablesEditor/Private/MultiUser/ConsoleVariableSyncData.h
+++ b/Engine/Plugins/Editor/ConsoleVariablesEditor/Source/ConsoleVariablesEditor/Private/MultiUser/ConsoleVariableSyncData.h
@@ -115,7 +115,11 @@ struct FConcertSetConsoleVariableEvent
 
 	/** The value to set the named console variable. */
 	UPROPERTY()
-	FString Value;
+    FString Value;
+
+	/** The flags for this console variable. */
+    UPROPERTY()
+	uint32 Flags = 0;
 };
 
 USTRUCT()
diff --git a/Engine/Plugins/Editor/ConsoleVariablesEditor/Source/ConsoleVariablesEditor/Private/Views/List/ConsoleVariablesEditorListRow.cpp b/Engine/Plugins/Editor/ConsoleVariablesEditor/Source/ConsoleVariablesEditor/Private/Views/List/ConsoleVariablesEditorListRow.cpp
--- a/Engine/Plugins/Editor/ConsoleVariablesEditor/Source/ConsoleVariablesEditor/Private/Views/List/ConsoleVariablesEditorListRow.cpp
+++ b/Engine/Plugins/Editor/ConsoleVariablesEditor/Source/ConsoleVariablesEditor/Private/Views/List/ConsoleVariablesEditorListRow.cpp
@@ -318,8 +318,10 @@ FReply FConsoleVariablesEditorListRow::OnActionButtonClicked()
 
 	FConsoleVariablesEditorModule& ConsoleVariablesEditorModule = FConsoleVariablesEditorModule::Get();
 
-	const FString& CommandName = GetCommandInfo().Pin()->Command;
-	const FString& StartupValue = GetCommandInfo().Pin()->StartupValueAsString;
+	TSharedPtr<FConsoleVariablesEditorCommandInfo> Info = GetCommandInfo().Pin();
+	const FString& CommandName = Info->Command;
+	const FString& StartupValue = Info->StartupValueAsString;
+	const EConsoleVariableFlags Source = Info->GetSource();
 	const TObjectPtr<UConsoleVariablesAsset> EditableAsset = ConsoleVariablesEditorModule.GetPresetAsset();
 	check(EditableAsset);
 
@@ -340,7 +342,7 @@ FReply FConsoleVariablesEditorListRow::OnActionButtonClicked()
 		{
 			EditableAsset->RemoveConsoleVariable(CommandName);
 
-			ConsoleVariablesEditorModule.SendMultiUserConsoleVariableChange(ERemoteCVarChangeType::Remove, CommandName, StartupValue);
+			ConsoleVariablesEditorModule.SendMultiUserConsoleVariableChange(ERemoteCVarChangeType::Remove, CommandName, StartupValue, Source);
 		}
 	}
 	else
@@ -349,7 +351,7 @@ FReply FConsoleVariablesEditorListRow::OnActionButtonClicked()
 
 		EditableAsset->RemoveConsoleVariable(CommandName);
 
-		ConsoleVariablesEditorModule.SendMultiUserConsoleVariableChange(ERemoteCVarChangeType::Remove, CommandName, StartupValue);
+		ConsoleVariablesEditorModule.SendMultiUserConsoleVariableChange(ERemoteCVarChangeType::Remove, CommandName, StartupValue, Source);
 		ListViewPtr.Pin()->RebuildListWithListMode(ListViewPtr.Pin()->GetListModelPtr().Pin()->GetListMode());
 	}
 
diff --git a/Engine/Plugins/Editor/ConsoleVariablesEditor/Source/ConsoleVariablesEditor/Private/Views/MainPanel/ConsoleVariablesEditorMainPanel.cpp b/Engine/Plugins/Editor/ConsoleVariablesEditor/Source/ConsoleVariablesEditor/Private/Views/MainPanel/ConsoleVariablesEditorMainPanel.cpp
--- a/Engine/Plugins/Editor/ConsoleVariablesEditor/Source/ConsoleVariablesEditor/Private/Views/MainPanel/ConsoleVariablesEditorMainPanel.cpp
+++ b/Engine/Plugins/Editor/ConsoleVariablesEditor/Source/ConsoleVariablesEditor/Private/Views/MainPanel/ConsoleVariablesEditorMainPanel.cpp
@@ -7,6 +7,7 @@
 #include "ConsoleVariablesEditorCommandInfo.h"
 #include "ConsoleVariablesEditorLog.h"
 #include "ConsoleVariablesEditorModule.h"
+#include "HAL/IConsoleManager.h"
 #include "MultiUser/ConsoleVariableSync.h"
 #include "Views/List/ConsoleVariablesEditorList.h"
 #include "Views/MainPanel/SConsoleVariablesEditorMainPanel.h"
@@ -74,12 +75,6 @@ void FConsoleVariablesEditorMainPanel::AddConsoleObjectToCurrentPreset(
 				ECheckBoxState::Checked
 			}
 		);
-		if (!bHadCommandInAsset)
-		{
-			GetConsoleVariablesModule().SendMultiUserConsoleVariableChange(
-				ERemoteCVarChangeType::Add,
-				TrimmedCommand, InValue);
-		}
 
 		if (GetEditorList().Pin()->GetListMode() == FConsoleVariablesEditorList::EConsoleVariablesEditorListMode::Preset)
 		{
@@ -336,11 +331,11 @@ void FConsoleVariablesEditorMainPanel::OnConnectionChanged(EConcertConnectionSta
 	}
 }
 
-void FConsoleVariablesEditorMainPanel::OnRemoteCvarChange(ERemoteCVarChangeType InChangeType, FString InName, FString InValue)
+void FConsoleVariablesEditorMainPanel::OnRemoteCvarChange(ERemoteCVarChangeType InChangeType, FString InName, FString InValue, EConsoleVariableFlags InFlags)
 {
 	FConsoleVariablesEditorModule& ConsoleVariablesEditorModule = FConsoleVariablesEditorModule::Get();
 
-	ConsoleVariablesEditorModule.OnRemoteCvarChanged(InChangeType, MoveTemp(InName), MoveTemp(InValue));
+	ConsoleVariablesEditorModule.OnRemoteCvarChanged(InChangeType, MoveTemp(InName), MoveTemp(InValue), InFlags);
 }
 
 void FConsoleVariablesEditorMainPanel::OnRemoteListItemCheckStateChange(const FString InName, ECheckBoxState InCheckedState)
diff --git a/Engine/Plugins/Editor/ConsoleVariablesEditor/Source/ConsoleVariablesEditor/Private/Views/MainPanel/ConsoleVariablesEditorMainPanel.h b/Engine/Plugins/Editor/ConsoleVariablesEditor/Source/ConsoleVariablesEditor/Private/Views/MainPanel/ConsoleVariablesEditorMainPanel.h
--- a/Engine/Plugins/Editor/ConsoleVariablesEditor/Source/ConsoleVariablesEditor/Private/Views/MainPanel/ConsoleVariablesEditorMainPanel.h
+++ b/Engine/Plugins/Editor/ConsoleVariablesEditor/Source/ConsoleVariablesEditor/Private/Views/MainPanel/ConsoleVariablesEditorMainPanel.h
@@ -90,7 +90,7 @@ private:
 	TSharedPtr<FConsoleVariablesEditorList> EditorList;
 
 	static void OnConnectionChanged(EConcertConnectionStatus Status);
-	static void OnRemoteCvarChange(ERemoteCVarChangeType InChangeType, FString InName, FString InValue);
+	static void OnRemoteCvarChange(ERemoteCVarChangeType InChangeType, FString InName, FString InValue, EConsoleVariableFlags InFlags);
 	static void OnRemoteListItemCheckStateChange(const FString InName, ECheckBoxState InCheckedState);
 	
 	void SendListItemCheckStateChange(const FString& InName, ECheckBoxState InCheckedState);
diff --git a/Engine/Plugins/Editor/ConsoleVariablesEditor/Source/ConsoleVariablesEditor/Public/ConsoleVariablesEditorModule.h b/Engine/Plugins/Editor/ConsoleVariablesEditor/Source/ConsoleVariablesEditor/Public/ConsoleVariablesEditorModule.h
--- a/Engine/Plugins/Editor/ConsoleVariablesEditor/Source/ConsoleVariablesEditor/Public/ConsoleVariablesEditorModule.h
+++ b/Engine/Plugins/Editor/ConsoleVariablesEditor/Source/ConsoleVariablesEditor/Public/ConsoleVariablesEditorModule.h
@@ -74,8 +74,8 @@ public:
 	/** Fills Global Search Asset's Saved Commands with variables matching the specified query. Returns false if no matches were found. */
 	bool PopulateGlobalSearchAssetWithVariablesMatchingTokens(const TArray<FString>& InTokens);
 
-	void SendMultiUserConsoleVariableChange(ERemoteCVarChangeType InChangeType, const FString& InVariableName, const FString& InValueAsString);
-	void OnRemoteCvarChanged(ERemoteCVarChangeType InChangeType, FString InName, FString InValue);
+	void SendMultiUserConsoleVariableChange(ERemoteCVarChangeType InChangeType, const FString& InVariableName, const FString& InValueAsString, EConsoleVariableFlags InFlags);
+	void OnRemoteCvarChanged(ERemoteCVarChangeType InChangeType, FString InName, FString InValue, EConsoleVariableFlags InFlags);
 
 	virtual void AddReferencedObjects( FReferenceCollector& Collector )  override;
 	virtual FString GetReferencerName() const override;
@@ -101,7 +101,7 @@ private:
 			else
 			{
 				int32& NewValue = InboundCommandTable.Add(TrackedCommand);
-				NewValue = 0;
+				NewValue = 1;
 			}
 		}
 
