Do not send console variables using the render thread flag as they do not represent the game thread cvar and cannot be applied on remote clients.

#rnx

diff --git a/Engine/Plugins/Editor/ConsoleVariablesEditor/Source/ConsoleVariablesEditor/Private/MultiUser/ConsoleVariableSync.cpp b/Engine/Plugins/Editor/ConsoleVariablesEditor/Source/ConsoleVariablesEditor/Private/MultiUser/ConsoleVariableSync.cpp
--- a/Engine/Plugins/Editor/ConsoleVariablesEditor/Source/ConsoleVariablesEditor/Private/MultiUser/ConsoleVariableSync.cpp
+++ b/Engine/Plugins/Editor/ConsoleVariablesEditor/Source/ConsoleVariablesEditor/Private/MultiUser/ConsoleVariableSync.cpp
@@ -436,6 +436,13 @@ struct FManagerImpl
 			return;
 		}
 
+		const bool bRenderThreadVersion = InFlags & ECVF_RenderThreadSafe;
+		if (bRenderThreadVersion)
+		{
+			// The receiver is never dealing with the render thread version of the CVar thus should not be applied to remote callers.
+			return;
+		}
+
 		uint32 OutFlags = static_cast<uint32>(InFlags);
 		FConcertSetConsoleVariableEvent OutEvent{ConvertChangeType(InChangeType), MoveTemp(InName), MoveTemp(InValue), InFlags};
 		Session->SendCustomEvent(OutEvent, Session->GetSessionClientEndpointIds(),
