[Console Variables Editor] Cosmetic UI/UX polish

diff --git a/Engine/Plugins/Editor/ConsoleVariablesEditor/Source/ConsoleVariablesEditor/Private/MultiUser/ConcertConsoleVariableSessionCustomization.cpp b/Engine/Plugins/Editor/ConsoleVariablesEditor/Source/ConsoleVariablesEditor/Private/MultiUser/ConcertConsoleVariableSessionCustomization.cpp
--- a/Engine/Plugins/Editor/ConsoleVariablesEditor/Source/ConsoleVariablesEditor/Private/MultiUser/ConcertConsoleVariableSessionCustomization.cpp
+++ b/Engine/Plugins/Editor/ConsoleVariablesEditor/Source/ConsoleVariablesEditor/Private/MultiUser/ConcertConsoleVariableSessionCustomization.cpp
@@ -2,6 +2,7 @@
 
 #include "ConcertConsoleVariableSessionCustomization.h"
 
+#include "ConsoleVariablesEditorStyle.h"
 #include "Containers/ContainersFwd.h"
 #include "Delegates/Delegate.h"
 #include "Framework/Views/ITypedTableView.h"
@@ -59,7 +60,11 @@ public:
 		SyncClient = MoveTemp(InSyncClient);
 		ClientSession = MoveTemp(InClientSession);
 		ClientSetting = MoveTemp(InClientInfo);
-		SMultiColumnTableRow<TSharedPtr<FConcertCVarDetails>>::Construct(FSuperRowType::FArguments(), InOwnerTableView);
+		SMultiColumnTableRow<TSharedPtr<FConcertCVarDetails>>::Construct(
+			FSuperRowType::FArguments()
+			.Style(&FAppStyle::Get().GetWidgetStyle<FTableRowStyle>("TableView.AlternatingRow")),
+			InOwnerTableView
+		);
 
 		SetToolTipText(MakeAttributeSP(this, &SConsoleVariableSessionRow::GetRowToolTip));
 	}
@@ -76,9 +81,10 @@ public:
 				.VAlign(VAlign_Center)
 				.Padding(FMargin(4.0f, 0.0f, 0.0f, 0.0f))
 				[
-					SNew(SImage)
-					.Image_Static(SConsoleVariableSessionRow::GetAvatarBrush)
+					SNew(STextBlock)
+					.Font(this, &SConsoleVariableSessionRow::GetAvatarFont)
 					.ColorAndOpacity(this, &SConsoleVariableSessionRow::GetAvatarColor)
+					.Text(FEditorFontGlyphs::Square)
 				]
 
 				// The client display name.
@@ -209,9 +215,28 @@ public:
 		return FText();
 	}
 
-	static const FSlateBrush* GetAvatarBrush()
+	FSlateFontInfo GetAvatarFont() const
 	{
-		return FAppStyle::Get().GetBrush("Icons.Toolbar.Stop");
+		static const FName ButtonIconSyle = TEXT("FontAwesome.10");
+		// This font is used to render a small square box filled with the avatar color.
+		FSlateFontInfo ClientIconFontInfo = FAppStyle::Get().GetFontStyle(ButtonIconSyle);
+		ClientIconFontInfo.Size = 8;
+		ClientIconFontInfo.OutlineSettings.OutlineSize = 1;
+
+		TSharedPtr<FConcertCVarDetails> ClientSettingPin = ClientSetting.Pin();
+		if (ClientSettingPin.IsValid())
+		{
+			FConcertSessionClientInfo& Client = ClientSettingPin->Details;
+			FLinearColor ClientOutlineColor = Client.ClientInfo.AvatarColor * 0.6f; // Make the font outline darker.
+			ClientOutlineColor.A = Client.ClientInfo.AvatarColor.A; // Put back the original alpha.
+			ClientIconFontInfo.OutlineSettings.OutlineColor = ClientOutlineColor;
+		}
+		else
+		{
+			ClientIconFontInfo.OutlineSettings.OutlineColor = FLinearColor(0.75, 0.75, 0.75); // This is an arbitrary color.
+		}
+
+		return ClientIconFontInfo;
 	}
 
 	FSlateColor GetAvatarColor() const
@@ -279,6 +304,17 @@ void FConcertConsoleVariableSessionCustomization::PopulateClientList()
 		FConcertCVarDetails Local = GetLocalSessionConcertClientAsSetting(Settings,Session);
 
 		Clients.Emplace(MakeShared<FConcertCVarDetails>(MoveTemp(Local)));
+
+		for (int32 DuplicateIndex = 0; DuplicateIndex < 5; ++DuplicateIndex)
+		{
+			auto DupeCliant = MakeShared<FConcertCVarDetails>(Local);
+			DupeCliant->Details.ClientInfo.AvatarColor = FLinearColor::MakeRandomSeededColor(DuplicateIndex);
+			DupeCliant->Details.ClientEndpointId = FGuid::NewGuid();
+			DupeCliant->Details.ClientInfo.DisplayName = DupeCliant->Details.ClientEndpointId.ToString();
+
+			Clients.Emplace(DupeCliant);
+		}
+
 		for (FConcertCVarDetails const& Client : Settings->RemoteDetails)
 		{
 			Clients.Emplace(MakeShared<FConcertCVarDetails>(Client));
@@ -341,7 +377,8 @@ void FConcertConsoleVariableSessionCustomization::SettingChange(const FConcertCV
 
 void FConcertConsoleVariableSessionCustomization::CustomizeDetails(IDetailLayoutBuilder& DetailLayout)
 {
-	IDetailCategoryBuilder& Settings = DetailLayout.EditCategory(TEXT("Multi-user Console Variable Settings"));
+	IDetailCategoryBuilder& Settings = DetailLayout.EditCategory(TEXT("MultiUserClientTransaction"));
+	Settings.SetDisplayName(LOCTEXT("MultiUserClientTransaction", "Multi-User Client Transaction"));
 
 	TSharedPtr<IPropertyHandle> SettingsProperty = DetailLayout.GetProperty(
 		GET_MEMBER_NAME_CHECKED(UConcertCVarConfig, LocalSettings));
@@ -354,7 +391,7 @@ void FConcertConsoleVariableSessionCustomization::CustomizeDetails(IDetailLayout
 
 	PopulateClientList();
 	{
-		FDetailWidgetRow& Row = Settings.AddCustomRow(LOCTEXT("MultiUserSettings", "Multi-user Client Settings"));
+		FDetailWidgetRow& Row = Settings.AddCustomRow(LOCTEXT("MultiUserSettings", "Multi-User Client Settings"));
 		auto HandleGenerateRow = [this](TSharedPtr<FConcertCVarDetails> InClientInfo,
 									const TSharedRef<STableViewBase>& OwnerTable) -> TSharedRef<ITableRow>
 		{
diff --git a/Engine/Plugins/Editor/ConsoleVariablesEditor/Source/ConsoleVariablesEditor/Private/MultiUser/ConcertConsoleVariableSyncCustomization.h b/Engine/Plugins/Editor/ConsoleVariablesEditor/Source/ConsoleVariablesEditor/Private/MultiUser/ConcertConsoleVariableSyncCustomization.h
--- a/Engine/Plugins/Editor/ConsoleVariablesEditor/Source/ConsoleVariablesEditor/Private/MultiUser/ConcertConsoleVariableSyncCustomization.h
+++ b/Engine/Plugins/Editor/ConsoleVariablesEditor/Source/ConsoleVariablesEditor/Private/MultiUser/ConcertConsoleVariableSyncCustomization.h
@@ -49,7 +49,8 @@ public:
 
 	void CustomizeDetails(IDetailLayoutBuilder& DetailLayout) override
 	{
-		IDetailCategoryBuilder& TakeSync = DetailLayout.EditCategory("Multi-user");
+		IDetailCategoryBuilder& TakeSync = DetailLayout.EditCategory("MultiUserSynchronization");
+		TakeSync.SetDisplayName(LOCTEXT("MultiUserSynchronization", "Multi-User Synchronization"));
 
 		TSharedPtr<IPropertyHandle> SyncTakeRecordingProperty = DetailLayout.GetProperty(
 			GET_MEMBER_NAME_CHECKED(UConcertCVarSynchronization, bSyncCVarTransactions));
diff --git a/Engine/Plugins/Editor/ConsoleVariablesEditor/Source/ConsoleVariablesEditor/Private/MultiUser/ConsoleVariableSyncData.h b/Engine/Plugins/Editor/ConsoleVariablesEditor/Source/ConsoleVariablesEditor/Private/MultiUser/ConsoleVariableSyncData.h
--- a/Engine/Plugins/Editor/ConsoleVariablesEditor/Source/ConsoleVariablesEditor/Private/MultiUser/ConsoleVariableSyncData.h
+++ b/Engine/Plugins/Editor/ConsoleVariablesEditor/Source/ConsoleVariablesEditor/Private/MultiUser/ConsoleVariableSyncData.h
@@ -17,7 +17,7 @@ class UConcertCVarSynchronization : public UObject
 public:
 	GENERATED_BODY()
 
-	UPROPERTY(config,EditAnywhere,BlueprintReadWrite,Category="Multi-user",DisplayName="Multi-user Console Variables Synchronization")
+	UPROPERTY(config,EditAnywhere,BlueprintReadWrite,Category="Multi-User",DisplayName="Multi-User Console Variables Synchronization")
 	bool bSyncCVarTransactions = true;
 };
 
diff --git a/Engine/Plugins/Editor/ConsoleVariablesEditor/Source/ConsoleVariablesEditor/Private/Views/List/SConsoleVariablesEditorList.cpp b/Engine/Plugins/Editor/ConsoleVariablesEditor/Source/ConsoleVariablesEditor/Private/Views/List/SConsoleVariablesEditorList.cpp
--- a/Engine/Plugins/Editor/ConsoleVariablesEditor/Source/ConsoleVariablesEditor/Private/Views/List/SConsoleVariablesEditorList.cpp
+++ b/Engine/Plugins/Editor/ConsoleVariablesEditor/Source/ConsoleVariablesEditor/Private/Views/List/SConsoleVariablesEditorList.cpp
@@ -24,6 +24,7 @@
 #include "Widgets/Input/SCheckBox.h"
 #include "Widgets/Input/SComboBox.h"
 #include "Widgets/Input/SSearchBox.h"
+#include "Widgets/Layout/SScaleBox.h"
 #include "Widgets/Layout/SWidgetSwitcher.h"
 #include "Widgets/Layout/SWrapBox.h"
 #include "Widgets/Text/SRichTextBlock.h"
@@ -198,7 +199,7 @@ void SConsoleVariablesEditorList::Construct(const FArguments& InArgs, TSharedRef
 			
 			+ SWidgetSwitcher::Slot()
 			.HAlign(HAlign_Fill)
-			.Padding(2.0f, 2.0f, 2.0f, 2.0f)
+			.Padding(0.0f, 2.0f, 0.0f, 2.0f)
 			[
 				SAssignNew(TreeViewPtr, STreeView<FConsoleVariablesEditorListRowPtr>)
 				.HeaderRow(HeaderRow)
@@ -226,7 +227,7 @@ void SConsoleVariablesEditorList::Construct(const FArguments& InArgs, TSharedRef
 			// For when no rows exist in view
 			+ SWidgetSwitcher::Slot()
 			.HAlign(HAlign_Fill)
-			.Padding(2.0f, 24.0f, 2.0f, 2.0f)
+			.Padding(0.0f, 24.0f, 0.0f, 2.0f)
 			[
 				SNew(SRichTextBlock)
 				.DecoratorStyleSet(&FAppStyle::Get())
@@ -876,7 +877,12 @@ TSharedPtr<SHeaderRow> SConsoleVariablesEditorList::GenerateHeaderRow()
 			.ShouldGenerateWidget(true)
 			.HeaderContent()
 			[
-				SNew(SBox)
+				SNew(SScaleBox)
+				[
+					SNew(SImage)
+						.Image(FAppStyle::Get().GetBrush("Icons.Delete"))
+						.ColorAndOpacity(FSlateColor::UseForeground())
+				]
 			]
 	);
 
diff --git a/Engine/Plugins/Editor/ConsoleVariablesEditor/Source/ConsoleVariablesEditor/Private/Views/List/SConsoleVariablesEditorListRow.cpp b/Engine/Plugins/Editor/ConsoleVariablesEditor/Source/ConsoleVariablesEditor/Private/Views/List/SConsoleVariablesEditorListRow.cpp
--- a/Engine/Plugins/Editor/ConsoleVariablesEditor/Source/ConsoleVariablesEditor/Private/Views/List/SConsoleVariablesEditorListRow.cpp
+++ b/Engine/Plugins/Editor/ConsoleVariablesEditor/Source/ConsoleVariablesEditor/Private/Views/List/SConsoleVariablesEditorListRow.cpp
@@ -78,11 +78,12 @@ void SConsoleVariablesEditorListRow::Construct(
 
 	SMultiColumnTableRow<FConsoleVariablesEditorListRowPtr>::Construct(
 		FSuperRowType::FArguments()
-		.Padding(1.0f)
+		.Padding(FMargin(0.0f, 1.0f))
 		.OnCanAcceptDrop(this, &SConsoleVariablesEditorListRow::HandleCanAcceptDrop)
 		.OnAcceptDrop(this, &SConsoleVariablesEditorListRow::HandleAcceptDrop)
 		.OnDragDetected(this, &SConsoleVariablesEditorListRow::HandleDragDetected)
-		.OnDragLeave(this, &SConsoleVariablesEditorListRow::HandleDragLeave),
+		.OnDragLeave(this, &SConsoleVariablesEditorListRow::HandleDragLeave)
+		.Style(&FAppStyle::Get().GetWidgetStyle<FTableRowStyle>("TableView.AlternatingRow")),
 		InOwnerTable
 	);
 
@@ -577,7 +578,7 @@ void SConsoleVariablesEditorListRowHoverWidgets::Construct(const FArguments& InA
 				SAssignNew(ActionButtonImage, SImage)
 				.Visibility(EVisibility::SelfHitTestInvisible)
 				.Image(ButtonImage)
-				.ColorAndOpacity(FSlateColor::UseForeground())
+				.ColorAndOpacity(this, &SConsoleVariablesEditorListRowHoverWidgets::GetActionButtonColorAndOpacity)
 			]
 		]
 	];
@@ -640,4 +641,14 @@ SConsoleVariablesEditorListRowHoverWidgets::~SConsoleVariablesEditorListRowHover
 	ActionButtonPtr.Reset();
 }
 
+FSlateColor SConsoleVariablesEditorListRowHoverWidgets::GetActionButtonColorAndOpacity() const
+{
+	if (IsHovered())
+	{
+		return FSlateColor::UseForeground();
+	}
+
+	return FSlateColor::UseSubduedForeground();
+}
+
 #undef LOCTEXT_NAMESPACE
diff --git a/Engine/Plugins/Editor/ConsoleVariablesEditor/Source/ConsoleVariablesEditor/Private/Views/List/SConsoleVariablesEditorListRow.h b/Engine/Plugins/Editor/ConsoleVariablesEditor/Source/ConsoleVariablesEditor/Private/Views/List/SConsoleVariablesEditorListRow.h
--- a/Engine/Plugins/Editor/ConsoleVariablesEditor/Source/ConsoleVariablesEditor/Private/Views/List/SConsoleVariablesEditorListRow.h
+++ b/Engine/Plugins/Editor/ConsoleVariablesEditor/Source/ConsoleVariablesEditor/Private/Views/List/SConsoleVariablesEditorListRow.h
@@ -85,7 +85,7 @@ private:
 	const FLinearColor FlashColor = FLinearColor::White;
 
 	/** The offset applied to text widgets so that the text aligns with the column header text */
-	float TextBlockLeftPadding = 3.0f;
+	float TextBlockLeftPadding = 5.0f;
 	
 	bool bIsHovered = false;
 };
@@ -110,6 +110,8 @@ public:
 	virtual ~SConsoleVariablesEditorListRowHoverWidgets() override;
 
 private:
+
+	FSlateColor GetActionButtonColorAndOpacity() const;
 	
 	TWeakPtr<FConsoleVariablesEditorListRow> Item;
 	
diff --git a/Engine/Plugins/Editor/ConsoleVariablesEditor/Source/ConsoleVariablesEditor/Private/Views/List/SConsoleVariablesEditorListValueInput.cpp b/Engine/Plugins/Editor/ConsoleVariablesEditor/Source/ConsoleVariablesEditor/Private/Views/List/SConsoleVariablesEditorListValueInput.cpp
--- a/Engine/Plugins/Editor/ConsoleVariablesEditor/Source/ConsoleVariablesEditor/Private/Views/List/SConsoleVariablesEditorListValueInput.cpp
+++ b/Engine/Plugins/Editor/ConsoleVariablesEditor/Source/ConsoleVariablesEditor/Private/Views/List/SConsoleVariablesEditorListValueInput.cpp
@@ -158,6 +158,7 @@ void SConsoleVariablesEditorListValueInput_Int::Construct(const FArguments& InAr
 	ChildSlot
 	[
 		SAssignNew(InputWidget, SSpinBox<int32>)
+		.Style(&FAppStyle::Get().GetWidgetStyle<FSpinBoxStyle>("NumericEntrySpinBox"))
 		.Value_Lambda([this]
 		{
 			check (Item.IsValid());
@@ -269,7 +270,7 @@ void SConsoleVariablesEditorListValueInput_String::Construct(const FArguments& I
 	
 	ChildSlot
 	[
-		SAssignNew(InputWidget, SEditableText)
+		SAssignNew(InputWidget, SEditableTextBox)
 		.Text_Lambda([this]
 		{
 			check (Item.IsValid());
@@ -422,7 +423,7 @@ void SConsoleVariablesEditorListValueInput_Command::Construct(const FArguments&
 		+SHorizontalBox::Slot()
 		.VAlign(VAlign_Center)
 		[
-			SAssignNew(InputText, SEditableText)
+			SAssignNew(InputText, SEditableTextBox)
 			.Text(FText::FromString(InSavedText))
 			.HintText(LOCTEXT("CommandValueTypeRowInputHintText", "Value..."))
 			.IsEnabled(this, &SConsoleVariablesEditorListValueInput::IsRowChecked)
@@ -432,7 +433,8 @@ void SConsoleVariablesEditorListValueInput_Command::Construct(const FArguments&
 		]
 
 		+SHorizontalBox::Slot()
-		.VAlign(VAlign_Center)
+		.Padding(FMargin(2.f, 0, 0, 0))
+		.VAlign(VAlign_Fill)
 		[
 			SAssignNew(InputWidget, SButton)
 			.OnClicked_Lambda([this] ()
@@ -459,9 +461,16 @@ void SConsoleVariablesEditorListValueInput_Command::Construct(const FArguments&
 			.IsEnabled(this, &SConsoleVariablesEditorListValueInput::IsRowChecked)
 			.ContentPadding(FMargin(0.f))
 			[
-				SNew(STextBlock)
-				.Justification(ETextJustify::Center)
-				.Text(LOCTEXT("ConsoleCommandExecutionButtonText","Execute"))
+				SNew(SHorizontalBox)
+
+				+SHorizontalBox::Slot()
+				.Padding(FMargin(2.f, 0, 0, 0))
+				.VAlign(VAlign_Center)
+				[
+					SNew(STextBlock)
+					.Justification(ETextJustify::Center)
+					.Text(LOCTEXT("ConsoleCommandExecutionButtonText", "Execute"))
+				]
 			]
 		]
 	];
diff --git a/Engine/Plugins/Editor/ConsoleVariablesEditor/Source/ConsoleVariablesEditor/Private/Views/List/SConsoleVariablesEditorListValueInput.h b/Engine/Plugins/Editor/ConsoleVariablesEditor/Source/ConsoleVariablesEditor/Private/Views/List/SConsoleVariablesEditorListValueInput.h
--- a/Engine/Plugins/Editor/ConsoleVariablesEditor/Source/ConsoleVariablesEditor/Private/Views/List/SConsoleVariablesEditorListValueInput.h
+++ b/Engine/Plugins/Editor/ConsoleVariablesEditor/Source/ConsoleVariablesEditor/Private/Views/List/SConsoleVariablesEditorListValueInput.h
@@ -3,6 +3,7 @@
 #pragma once
 
 #include "Widgets/Input/SButton.h"
+#include "Widgets/Input/SEditableTextBox.h"
 #include "Widgets/Input/SSpinBox.h"
 #include "Widgets/SCompoundWidget.h"
 #include "Widgets/Text/STextBlock.h"
@@ -99,7 +100,7 @@ public:
 
 private:
 	
-	TSharedPtr<SEditableText> InputWidget;
+	TSharedPtr<SEditableTextBox> InputWidget;
 };
 
 class SConsoleVariablesEditorListValueInput_Bool : public SConsoleVariablesEditorListValueInput
@@ -159,5 +160,5 @@ public:
 private:
 	
 	TSharedPtr<SButton> InputWidget;
-	TSharedPtr<SEditableText> InputText;
+	TSharedPtr<SEditableTextBox> InputText;
 };
