Use the ECVF_SetByMask flag before assigning new console variable values

diff --git a/Engine/Plugins/Editor/ConsoleVariablesEditor/Source/ConsoleVariablesEditor/Private/ConsoleVariablesEditorModule.cpp b/Engine/Plugins/Editor/ConsoleVariablesEditor/Source/ConsoleVariablesEditor/Private/ConsoleVariablesEditorModule.cpp
--- a/Engine/Plugins/Editor/ConsoleVariablesEditor/Source/ConsoleVariablesEditor/Private/ConsoleVariablesEditorModule.cpp
+++ b/Engine/Plugins/Editor/ConsoleVariablesEditor/Source/ConsoleVariablesEditor/Private/ConsoleVariablesEditorModule.cpp
@@ -302,7 +302,7 @@ void FConsoleVariablesEditorModule::OnRemoteCvarChanged(ERemoteCVarChangeType In
 		FConsoleVariablesEditorCommandInfo Info(InName);
 		if (IConsoleVariable* AsVariable = Info.GetConsoleVariablePtr())
 		{
-			AsVariable->Set(*InValue, InFlags);
+			AsVariable->Set(*InValue, (EConsoleVariableFlags)(InFlags & ECVF_SetByMask));
 		}
 		else
 		{
