[Console Variables Editor] Enable multi-select actions

diff --git a/Engine/Plugins/Editor/ConsoleVariablesEditor/Source/ConsoleVariablesEditor/Private/Views/List/ConsoleVariablesEditorListRow.cpp b/Engine/Plugins/Editor/ConsoleVariablesEditor/Source/ConsoleVariablesEditor/Private/Views/List/ConsoleVariablesEditorListRow.cpp
--- a/Engine/Plugins/Editor/ConsoleVariablesEditor/Source/ConsoleVariablesEditor/Private/Views/List/ConsoleVariablesEditorListRow.cpp
+++ b/Engine/Plugins/Editor/ConsoleVariablesEditor/Source/ConsoleVariablesEditor/Private/Views/List/ConsoleVariablesEditorListRow.cpp
@@ -308,28 +308,56 @@ FReply FConsoleVariablesEditorListRow::OnActionButtonClicked()
 		return FReply::Unhandled();
 	}
 
+	FConsoleVariablesEditorModule& ConsoleVariablesEditorModule = FConsoleVariablesEditorModule::Get();
+	const TObjectPtr<UConsoleVariablesAsset> EditableAsset = ConsoleVariablesEditorModule.GetPresetAsset();
+	check(EditableAsset);
+
 	const FConsoleVariablesEditorList::EConsoleVariablesEditorListMode ListMode =
 		bIsListValid
 			? GetListViewPtr().Pin()->GetListModelPtr().Pin()->GetListMode()
 			: FConsoleVariablesEditorList::EConsoleVariablesEditorListMode::Preset;
 
-	const bool bIsGlobalSearch =
-		ListMode == FConsoleVariablesEditorList::EConsoleVariablesEditorListMode::GlobalSearch;
-
-	FConsoleVariablesEditorModule& ConsoleVariablesEditorModule = FConsoleVariablesEditorModule::Get();
-
-	TSharedPtr<FConsoleVariablesEditorCommandInfo> Info = GetCommandInfo().Pin();
-	const FString& CommandName = Info->Command;
-	const FString& StartupValue = Info->StartupValueAsString;
-	const EConsoleVariableFlags Source = Info->GetSource();
-	const TObjectPtr<UConsoleVariablesAsset> EditableAsset = ConsoleVariablesEditorModule.GetPresetAsset();
-	check(EditableAsset);
+	bool bWillRemove = true;
 
+	// In global search mode, if this command isn't already in the preset, we want to add the preset instead
+	const bool bIsGlobalSearch = ListMode == FConsoleVariablesEditorList::EConsoleVariablesEditorListMode::GlobalSearch;
 	if (bIsGlobalSearch)
 	{
 		FConsoleVariablesEditorAssetSaveData MatchingData;
-		if (!ConsoleVariablesEditorModule.GetPresetAsset()->FindSavedDataByCommandString(CommandName, MatchingData, ESearchCase::IgnoreCase))
+		const FString& ClickedCommandName = GetCommandInfo().Pin()->Command;
+		if (!ConsoleVariablesEditorModule.GetPresetAsset()->FindSavedDataByCommandString(ClickedCommandName, MatchingData, ESearchCase::IgnoreCase))
+		{
+			bWillRemove = false;
+		}
+	}
+
+	for (const FConsoleVariablesEditorListRowPtr& RowPtr : GetRowsAffectedByActions())
+	{
+		if (!RowPtr.IsValid())
 		{
+			continue;
+		}
+
+		TSharedPtr<FConsoleVariablesEditorCommandInfo> Info = RowPtr->GetCommandInfo().Pin();
+		const FString& CommandName = Info->Command;
+
+		if (bWillRemove)
+		{
+			const FString& StartupValue = Info->StartupValueAsString;
+			const EConsoleVariableFlags Source = Info->GetSource();
+
+			RowPtr->ResetToStartupValueAndSource();
+			EditableAsset->RemoveConsoleVariable(CommandName);
+
+			ConsoleVariablesEditorModule.SendMultiUserConsoleVariableChange(ERemoteCVarChangeType::Remove, CommandName, StartupValue, Source);
+
+			// Refresh preset list + numbers now that the row has been removed
+			ListViewPtr.Pin()->RebuildListWithListMode(ListViewPtr.Pin()->GetListModelPtr().Pin()->GetListMode());
+		}
+		else
+		{
+			ensureMsgf(bIsGlobalSearch, TEXT("Add to preset operation should only occur in global search mode"));
+
 			EditableAsset->AddOrSetConsoleObjectSavedData(
 				{
 					CommandName,
@@ -338,24 +366,51 @@ FReply FConsoleVariablesEditorListRow::OnActionButtonClicked()
 				}
 			);
 		}
-		else
-		{
-			EditableAsset->RemoveConsoleVariable(CommandName);
+	}
 
-			ConsoleVariablesEditorModule.SendMultiUserConsoleVariableChange(ERemoteCVarChangeType::Remove, CommandName, StartupValue, Source);
-		}
+	return FReply::Handled();
+}
+
+bool FConsoleVariablesEditorListRow::IsGlobalSearch() const
+{
+	if (!ListViewPtr.IsValid())
+	{
+		return false;
 	}
-	else
+
+	const TWeakPtr<FConsoleVariablesEditorList> ListModelPtr = ListViewPtr.Pin()->GetListModelPtr();
+	if (!ListModelPtr.IsValid())
 	{
-		ResetToStartupValueAndSource();
+		return false;
+	}
+
+	return ListModelPtr.Pin()->GetListMode() == FConsoleVariablesEditorList::EConsoleVariablesEditorListMode::GlobalSearch;
+}
+
+bool FConsoleVariablesEditorListRow::IsInPreset() const
+{
+	const FString& CommandName = GetCommandInfo().Pin()->Command;
+	FConsoleVariablesEditorAssetSaveData MatchingData;
 
-		EditableAsset->RemoveConsoleVariable(CommandName);
+	const FConsoleVariablesEditorModule& ConsoleVariablesEditorModule = FConsoleVariablesEditorModule::Get();
+	return ConsoleVariablesEditorModule.GetPresetAsset()->FindSavedDataByCommandString(CommandName, MatchingData);
+}
+
+TArray<FConsoleVariablesEditorListRowPtr> FConsoleVariablesEditorListRow::GetRowsAffectedByActions()
+{
+	TArray<FConsoleVariablesEditorListRowPtr> AffectedRows = GetSelectedTreeViewItems();
+	const bool bIsRowSelected = AffectedRows.ContainsByPredicate([this](const FConsoleVariablesEditorListRowPtr& RowPtr)
+	{
+		return RowPtr.Get() == this;
+	});
 
-		ConsoleVariablesEditorModule.SendMultiUserConsoleVariableChange(ERemoteCVarChangeType::Remove, CommandName, StartupValue, Source);
-		ListViewPtr.Pin()->RebuildListWithListMode(ListViewPtr.Pin()->GetListModelPtr().Pin()->GetListMode());
+	if (!bIsRowSelected)
+	{
+		AffectedRows.Empty(1);
+		AffectedRows.Add(AsShared());
 	}
 
-	return FReply::Handled();
+	return AffectedRows;
 }
 
 void FConsoleVariablesEditorListRow::ResetToPresetValue()
diff --git a/Engine/Plugins/Editor/ConsoleVariablesEditor/Source/ConsoleVariablesEditor/Private/Views/List/ConsoleVariablesEditorListRow.h b/Engine/Plugins/Editor/ConsoleVariablesEditor/Source/ConsoleVariablesEditor/Private/Views/List/ConsoleVariablesEditorListRow.h
--- a/Engine/Plugins/Editor/ConsoleVariablesEditor/Source/ConsoleVariablesEditor/Private/Views/List/ConsoleVariablesEditorListRow.h
+++ b/Engine/Plugins/Editor/ConsoleVariablesEditor/Source/ConsoleVariablesEditor/Private/Views/List/ConsoleVariablesEditorListRow.h
@@ -127,6 +127,19 @@ struct FConsoleVariablesEditorListRow final : TSharedFromThis<FConsoleVariablesE
 
 	[[nodiscard]] TArray<FConsoleVariablesEditorListRowPtr> GetSelectedTreeViewItems() const;
 
+	/** Return true if the row is being displayed in Global Search mode. */
+	[[nodiscard]] bool IsGlobalSearch() const;
+
+	/** Return true if the command represented by the row is included in the current preset. */
+	[[nodiscard]] bool IsInPreset() const;
+
+	/**
+	 * Get an array of all rows that should be affected if an action is performed on this row.
+	 * This takes multi-selection into account and imitates the Outliner visibility widget's behavior by only
+	 * returning the row itself if the row is outside of the current selection.
+	 */
+	[[nodiscard]] TArray<FConsoleVariablesEditorListRowPtr> GetRowsAffectedByActions();
+
 	FReply OnActionButtonClicked();
 	
 	void ResetToPresetValue();
diff --git a/Engine/Plugins/Editor/ConsoleVariablesEditor/Source/ConsoleVariablesEditor/Private/Views/List/SConsoleVariablesEditorList.cpp b/Engine/Plugins/Editor/ConsoleVariablesEditor/Source/ConsoleVariablesEditor/Private/Views/List/SConsoleVariablesEditorList.cpp
--- a/Engine/Plugins/Editor/ConsoleVariablesEditor/Source/ConsoleVariablesEditor/Private/Views/List/SConsoleVariablesEditorList.cpp
+++ b/Engine/Plugins/Editor/ConsoleVariablesEditor/Source/ConsoleVariablesEditor/Private/Views/List/SConsoleVariablesEditorList.cpp
@@ -880,7 +880,7 @@ TSharedPtr<SHeaderRow> SConsoleVariablesEditorList::GenerateHeaderRow()
 				SNew(SScaleBox)
 				[
 					SNew(SImage)
-						.Image(FAppStyle::Get().GetBrush("Icons.Delete"))
+						.Image(this, &SConsoleVariablesEditorList::GetActionColumnImage)
 						.ColorAndOpacity(FSlateColor::UseForeground())
 				]
 			]
@@ -1306,4 +1306,20 @@ void SConsoleVariablesEditorList::SetChildExpansionRecursively(const FConsoleVar
 	}
 }
 
+
+const FSlateBrush* SConsoleVariablesEditorList::GetActionColumnImage() const
+{
+	const FConsoleVariablesEditorList::EConsoleVariablesEditorListMode ListMode =
+		ListModelPtr.IsValid()
+		? ListModelPtr.Pin()->GetListMode()
+		: FConsoleVariablesEditorList::EConsoleVariablesEditorListMode::Preset;
+
+	if (ListMode == FConsoleVariablesEditorList::EConsoleVariablesEditorListMode::GlobalSearch)
+	{
+		return FAppStyle::Get().GetBrush("Icons.Star");
+	}
+
+	return FAppStyle::Get().GetBrush("Icons.Delete");
+}
+
 #undef LOCTEXT_NAMESPACE
diff --git a/Engine/Plugins/Editor/ConsoleVariablesEditor/Source/ConsoleVariablesEditor/Private/Views/List/SConsoleVariablesEditorList.h b/Engine/Plugins/Editor/ConsoleVariablesEditor/Source/ConsoleVariablesEditor/Private/Views/List/SConsoleVariablesEditorList.h
--- a/Engine/Plugins/Editor/ConsoleVariablesEditor/Source/ConsoleVariablesEditor/Private/Views/List/SConsoleVariablesEditorList.h
+++ b/Engine/Plugins/Editor/ConsoleVariablesEditor/Source/ConsoleVariablesEditor/Private/Views/List/SConsoleVariablesEditorList.h
@@ -161,6 +161,8 @@ private:
 
 	void SetChildExpansionRecursively(const FConsoleVariablesEditorListRowPtr& InRow, const bool bNewIsExpanded) const;
 
+	const FSlateBrush* GetActionColumnImage() const;
+
 	TArray<TSharedRef<IConsoleVariablesEditorListFilter>> ShowFilters;
 
 	TSharedPtr<STreeView<FConsoleVariablesEditorListRowPtr>> TreeViewPtr;
diff --git a/Engine/Plugins/Editor/ConsoleVariablesEditor/Source/ConsoleVariablesEditor/Private/Views/List/SConsoleVariablesEditorListRow.cpp b/Engine/Plugins/Editor/ConsoleVariablesEditor/Source/ConsoleVariablesEditor/Private/Views/List/SConsoleVariablesEditorListRow.cpp
--- a/Engine/Plugins/Editor/ConsoleVariablesEditor/Source/ConsoleVariablesEditor/Private/Views/List/SConsoleVariablesEditorListRow.cpp
+++ b/Engine/Plugins/Editor/ConsoleVariablesEditor/Source/ConsoleVariablesEditor/Private/Views/List/SConsoleVariablesEditorListRow.cpp
@@ -130,31 +130,6 @@ TSharedRef<SWidget> SConsoleVariablesEditorListRow::GenerateWidgetForColumn(cons
 			];
 }
 
-void SConsoleVariablesEditorListRow::OnMouseEnter(const FGeometry& MyGeometry, const FPointerEvent& MouseEvent)
-{
-	bIsHovered = true;
-
-	if (HoverableWidgetsPtr.IsValid())
-	{
-		HoverableWidgetsPtr->DetermineButtonImageAndTooltip();
-		HoverableWidgetsPtr->SetVisibility(EVisibility::SelfHitTestInvisible);
-	}
-
-	SMultiColumnTableRow<FConsoleVariablesEditorListRowPtr>::OnMouseEnter(MyGeometry, MouseEvent);
-}
-
-void SConsoleVariablesEditorListRow::OnMouseLeave(const FPointerEvent& MouseEvent)
-{
-	bIsHovered = false;
-
-	if (HoverableWidgetsPtr.IsValid())
-	{
-		HoverableWidgetsPtr->SetVisibility(EVisibility::Collapsed);
-	}
-
-	SMultiColumnTableRow<FConsoleVariablesEditorListRowPtr>::OnMouseLeave(MouseEvent);
-}
-
 SConsoleVariablesEditorListRow::~SConsoleVariablesEditorListRow()
 {
 	Item.Reset();
@@ -302,6 +277,19 @@ void SConsoleVariablesEditorListRow::FlashRow()
 	FlashAnimation.Play(this->AsShared());
 }
 
+EVisibility SConsoleVariablesEditorListRow::GetHoverableWidgetsVisibility() const
+{
+	const FConsoleVariablesEditorListRowPtr RowPtr = Item.Pin();
+
+	if (IsHovered() || IsSelected() ||
+		(RowPtr.IsValid() && RowPtr->IsGlobalSearch() && RowPtr->IsInPreset()))
+	{
+		return EVisibility::Visible;
+	}
+
+	return EVisibility::Collapsed;
+}
+
 EVisibility SConsoleVariablesEditorListRow::GetFlashImageVisibility() const
 {
 	return FlashAnimation.IsPlaying() ? EVisibility::SelfHitTestInvisible : EVisibility::Hidden;
@@ -408,8 +396,8 @@ TSharedRef<SWidget> SConsoleVariablesEditorListRow::GenerateCells(const FName& I
 	}
 	if (InColumnName.IsEqual(SConsoleVariablesEditorList::ActionButtonColumnName))
 	{
-		return SAssignNew(HoverableWidgetsPtr, SConsoleVariablesEditorListRowHoverWidgets, Item)
-				.Visibility(EVisibility::Collapsed);
+		return SNew(SConsoleVariablesEditorListRowHoverWidgets, Item)
+				.Visibility(this, &SConsoleVariablesEditorListRow::GetHoverableWidgetsVisibility);
 	}
 
 	return SNullWidget::NullWidget;
@@ -434,27 +422,46 @@ ECheckBoxState SConsoleVariablesEditorListRow::GetCheckboxState() const
 
 void SConsoleVariablesEditorListRow::OnCheckboxStateChange(const ECheckBoxState InNewState) const
 {
-	if (const TSharedPtr<FConsoleVariablesEditorListRow> PinnedItem = Item.Pin())
+	const TSharedPtr<FConsoleVariablesEditorListRow> PinnedItem = Item.Pin();
+
+	if (!PinnedItem.IsValid())
 	{
-		PinnedItem->SetWidgetCheckedState(InNewState, true);
+		return;
+	}
 
-		if (PinnedItem->GetRowType() == FConsoleVariablesEditorListRow::SingleCommand)
+	for (const FConsoleVariablesEditorListRowPtr& RowPtr : PinnedItem->GetRowsAffectedByActions())
+	{
+		if (!RowPtr.IsValid())
 		{
-			if (PinnedItem->IsRowChecked())
+			continue;
+		}
+
+		if (RowPtr->GetCommandInfo().Pin()->ObjectType !=
+			FConsoleVariablesEditorCommandInfo::EConsoleObjectType::Variable)
+		{
+			// Only variables can be checked/unchecked
+			continue;
+		}
+
+		RowPtr->SetWidgetCheckedState(InNewState, true);
+
+		if (RowPtr->GetRowType() == FConsoleVariablesEditorListRow::SingleCommand)
+		{
+			if (RowPtr->IsRowChecked())
 			{
-				PinnedItem->GetCommandInfo().Pin()->ExecuteCommand(PinnedItem->GetCachedValue());
+				RowPtr->GetCommandInfo().Pin()->ExecuteCommand(RowPtr->GetCachedValue());
 			}
 			else
 			{
-				PinnedItem->ResetToStartupValueAndSource();
+				RowPtr->ResetToStartupValueAndSource();
 			}
 		}
+	}
 
-		if (const TWeakPtr<SConsoleVariablesEditorList> ListView = PinnedItem->GetListViewPtr(); ListView.IsValid())
-		{
-			const TSharedPtr<SConsoleVariablesEditorList> PinnedListView = ListView.Pin();
-			PinnedListView->OnListItemCheckBoxStateChange(InNewState);
-		}
+	if (const TWeakPtr<SConsoleVariablesEditorList> ListView = PinnedItem->GetListViewPtr(); ListView.IsValid())
+	{
+		const TSharedPtr<SConsoleVariablesEditorList> PinnedListView = ListView.Pin();
+		PinnedListView->OnListItemCheckBoxStateChange(InNewState);
 	}
 }
 
@@ -503,15 +510,15 @@ TSharedRef<SWidget> SConsoleVariablesEditorListRow::GenerateValueCellWidget(
 			.ButtonStyle(&FAppStyle::Get().GetWidgetStyle<FButtonStyle>("NoBorder"))
 			.ContentPadding(0)
 			.Visibility_Lambda([this, PinnedItem]()
-             {
-	             if (PinnedItem.IsValid() && PinnedItem->GetRowType() ==
-		             FConsoleVariablesEditorListRow::SingleCommand)
-	             {
-		             const TSharedPtr<FConsoleVariablesEditorCommandInfo> CommandInfo =
-			             PinnedItem->GetCommandInfo().Pin();
+            {
+	            if (PinnedItem.IsValid() && PinnedItem->GetRowType() ==
+		            FConsoleVariablesEditorListRow::SingleCommand)
+	            {
+		            const TSharedPtr<FConsoleVariablesEditorCommandInfo> CommandInfo =
+			            PinnedItem->GetCommandInfo().Pin();
 		            
-		            if (CommandInfo.IsValid() &&
-			            CommandInfo->ObjectType == FConsoleVariablesEditorCommandInfo::EConsoleObjectType::Variable)
+					if (CommandInfo.IsValid() &&
+						CommandInfo->ObjectType == FConsoleVariablesEditorCommandInfo::EConsoleObjectType::Variable)
 		            {
 			            PinnedItem->SetDoesCurrentValueDifferFromPresetValue(
 							 CommandInfo->IsCurrentValueDifferentFromInputValue(
@@ -522,16 +529,28 @@ TSharedRef<SWidget> SConsoleVariablesEditorListRow::GenerateValueCellWidget(
 									? EVisibility::Visible
 									: EVisibility::Collapsed;
 		            }
-	             }
+	            }
 
-	             return EVisibility::Collapsed;
-             })
+	            return EVisibility::Collapsed;
+            })
 			.OnClicked_Lambda([this, PinnedItem]()
-             {
-	             PinnedItem->ResetToPresetValue();
+            {
+				for (const FConsoleVariablesEditorListRowPtr& RowPtr : PinnedItem->GetRowsAffectedByActions())
+				{
+					const TSharedPtr<FConsoleVariablesEditorCommandInfo> CommandInfo =
+						PinnedItem->GetCommandInfo().Pin();
 
-	             return FReply::Handled();
-             })
+					if (!CommandInfo.IsValid() ||
+						CommandInfo->ObjectType != FConsoleVariablesEditorCommandInfo::EConsoleObjectType::Variable)
+					{
+						continue;
+					}
+
+					RowPtr->ResetToPresetValue();
+				}
+
+	            return FReply::Handled();
+            })
 			[
 				SNew(SImage)
 				.Image(FAppStyle::Get().GetBrush("PropertyWindow.DiffersFromDefault"))
@@ -558,97 +577,67 @@ void SConsoleVariablesEditorListRowHoverWidgets::Construct(const FArguments& InA
 	ChildSlot
 	[
 		// Action Button
-		SAssignNew(ActionButtonPtr, SButton)
+		SNew(SButton)
 		.ButtonStyle(&FAppStyle::Get().GetWidgetStyle<FButtonStyle>("NoBorder"))
 		.Visibility(EVisibility::Visible)
-		.ToolTipText(ButtonTooltip)
+		.ToolTipText(this, &SConsoleVariablesEditorListRowHoverWidgets::GetButtonTooltip)
 		.ButtonColorAndOpacity(FStyleColors::Transparent)
 		.ContentPadding(0.f)
 		.OnClicked_Lambda([this]()
 		{
-			SetVisibility(EVisibility::Collapsed);
-			FReply Reply = Item.Pin()->OnActionButtonClicked();
-			DetermineButtonImageAndTooltip();
-			SetVisibility(EVisibility::SelfHitTestInvisible);
-			return Reply;
+			return Item.Pin()->OnActionButtonClicked();
 		})
 		[
 			SNew(SScaleBox)
 			[
-				SAssignNew(ActionButtonImage, SImage)
+				SNew(SImage)
 				.Visibility(EVisibility::SelfHitTestInvisible)
-				.Image(ButtonImage)
+				.Image(this, &SConsoleVariablesEditorListRowHoverWidgets::GetButtonImage)
 				.ColorAndOpacity(this, &SConsoleVariablesEditorListRowHoverWidgets::GetActionButtonColorAndOpacity)
 			]
 		]
 	];
 }
 
-void SConsoleVariablesEditorListRowHoverWidgets::DetermineButtonImageAndTooltip()
+SConsoleVariablesEditorListRowHoverWidgets::~SConsoleVariablesEditorListRowHoverWidgets()
 {
-	const bool bIsListValid =
-		Item.Pin()->GetListViewPtr().IsValid() && Item.Pin()->GetListViewPtr().Pin()->GetListModelPtr().IsValid();
-
-	const FConsoleVariablesEditorList::EConsoleVariablesEditorListMode ListMode =
-		bIsListValid
-			? Item.Pin()->GetListViewPtr().Pin()->GetListModelPtr().Pin()->GetListMode()
-			: FConsoleVariablesEditorList::EConsoleVariablesEditorListMode::Preset;
-
-	const bool bIsGlobalSearch =
-		ListMode == FConsoleVariablesEditorList::EConsoleVariablesEditorListMode::GlobalSearch;
-
-	if (bIsGlobalSearch)
-	{
-		FText ButtonTooltip = LOCTEXT("AddCvarToPresetTooltip", "Add this cvar to your current preset.");
-		const FSlateBrush* ButtonImage = FConsoleVariablesEditorStyle::Get().GetBrush("ConsoleVariables.Favorite.Outline.Small");
-		
-		const FConsoleVariablesEditorModule& ConsoleVariablesEditorModule = FConsoleVariablesEditorModule::Get();
-		
-		const FString& CommandName = Item.Pin()->GetCommandInfo().Pin()->Command;
-		FConsoleVariablesEditorAssetSaveData MatchingData;
-
-		// If the item exists in the current preset already
-		if (ConsoleVariablesEditorModule.GetPresetAsset()->FindSavedDataByCommandString(CommandName, MatchingData))
-		{
-			ButtonTooltip = LOCTEXT("RemoveCvarFromPresetTooltip", "Remove this cvar from your current preset.");
-			ButtonImage = FAppStyle::Get().GetBrush("Icons.Star");
-		}
-
-		ActionButtonPtr->SetToolTipText(ButtonTooltip);
-		ActionButtonImage->SetImage(ButtonImage);
-	}
+	Item.Reset();
 }
 
-void SConsoleVariablesEditorListRowHoverWidgets::OnMouseEnter(const FGeometry& MyGeometry,
-                                                              const FPointerEvent& MouseEvent)
+FSlateColor SConsoleVariablesEditorListRowHoverWidgets::GetActionButtonColorAndOpacity() const
 {
-	SCompoundWidget::OnMouseEnter(MyGeometry, MouseEvent);
+	if (IsHovered())
+	{
+		return FSlateColor::UseForeground();
+	}
 
-	ActionButtonPtr->SetBorderBackgroundColor(FLinearColor(0.f, 0.f, 0.f, 0.4f));
+	return FSlateColor::UseSubduedForeground();
 }
 
-void SConsoleVariablesEditorListRowHoverWidgets::OnMouseLeave(const FPointerEvent& MouseEvent)
+const FSlateBrush* SConsoleVariablesEditorListRowHoverWidgets::GetButtonImage() const
 {
-	SCompoundWidget::OnMouseLeave(MouseEvent);
-
-	ActionButtonPtr->SetBorderBackgroundColor(FLinearColor(0.f, 0.f, 0.f, 0.f));
-}
+	const FConsoleVariablesEditorListRowPtr RowPtr = Item.Pin();
+	if (RowPtr.IsValid() && RowPtr->IsGlobalSearch())
+	{
+		return RowPtr->IsInPreset()
+			? FAppStyle::Get().GetBrush("Icons.Star")
+			: FConsoleVariablesEditorStyle::Get().GetBrush("ConsoleVariables.Favorite.Outline.Small");
+	}
 
-SConsoleVariablesEditorListRowHoverWidgets::~SConsoleVariablesEditorListRowHoverWidgets()
-{
-	Item.Reset();
-	
-	ActionButtonPtr.Reset();
+	return FAppStyle::Get().GetBrush("Icons.Delete");
 }
 
-FSlateColor SConsoleVariablesEditorListRowHoverWidgets::GetActionButtonColorAndOpacity() const
+FText SConsoleVariablesEditorListRowHoverWidgets::GetButtonTooltip() const
 {
-	if (IsHovered())
+	const FConsoleVariablesEditorListRowPtr RowPtr = Item.Pin();
+	if (RowPtr.IsValid() && RowPtr->IsGlobalSearch())
 	{
-		return FSlateColor::UseForeground();
+		return RowPtr->IsInPreset()
+			? LOCTEXT("RemoveCvarFromPresetTooltip", "Remove this cvar from your current preset.")
+			: LOCTEXT("AddCvarToPresetTooltip", "Add this cvar to your current preset.");
 	}
 
-	return FSlateColor::UseSubduedForeground();
+	return LOCTEXT("RemoveCvarTooltip", "Remove cvar from this list and reset its value to the startup value.");
 }
 
 #undef LOCTEXT_NAMESPACE
diff --git a/Engine/Plugins/Editor/ConsoleVariablesEditor/Source/ConsoleVariablesEditor/Private/Views/List/SConsoleVariablesEditorListRow.h b/Engine/Plugins/Editor/ConsoleVariablesEditor/Source/ConsoleVariablesEditor/Private/Views/List/SConsoleVariablesEditorListRow.h
--- a/Engine/Plugins/Editor/ConsoleVariablesEditor/Source/ConsoleVariablesEditor/Private/Views/List/SConsoleVariablesEditorListRow.h
+++ b/Engine/Plugins/Editor/ConsoleVariablesEditor/Source/ConsoleVariablesEditor/Private/Views/List/SConsoleVariablesEditorListRow.h
@@ -28,11 +28,6 @@ public:
 	void Construct(const FArguments& InArgs, const TSharedRef<STableViewBase>& InOwnerTable, const TWeakPtr<FConsoleVariablesEditorListRow> InRow);
 
 	virtual TSharedRef<SWidget> GenerateWidgetForColumn(const FName& InColumnName) override;
-
-	// Begin SWidget
-	virtual void OnMouseEnter(const FGeometry& MyGeometry, const FPointerEvent& MouseEvent) override;
-	virtual void OnMouseLeave(const FPointerEvent& MouseEvent) override;
-	// End SWidget
 	
 	virtual ~SConsoleVariablesEditorListRow() override;	
 
@@ -42,7 +37,8 @@ public:
 	FReply HandleAcceptDrop(const FDragDropEvent& DragDropEvent, EItemDropZone DropZone, FConsoleVariablesEditorListRowPtr TargetItem);
 	
 	void FlashRow();
-	
+
+	EVisibility GetHoverableWidgetsVisibility() const;
 	EVisibility GetFlashImageVisibility() const;
 	FSlateColor GetFlashImageColorAndOpacity() const;
 
@@ -86,8 +82,6 @@ private:
 
 	/** The offset applied to text widgets so that the text aligns with the column header text */
 	float TextBlockLeftPadding = 5.0f;
-	
-	bool bIsHovered = false;
 };
 
 class SConsoleVariablesEditorListRowHoverWidgets : public SCompoundWidget
@@ -101,20 +95,13 @@ public:
 	
 	void Construct(const FArguments& InArgs, const TWeakPtr<FConsoleVariablesEditorListRow> InRow);
 
-	void DetermineButtonImageAndTooltip();
-
-	virtual void OnMouseEnter(const FGeometry& MyGeometry, const FPointerEvent& MouseEvent) override;
-
-	virtual void OnMouseLeave(const FPointerEvent& MouseEvent) override;
-
 	virtual ~SConsoleVariablesEditorListRowHoverWidgets() override;
 
 private:
 
 	FSlateColor GetActionButtonColorAndOpacity() const;
+	const FSlateBrush* GetButtonImage() const;
+	FText GetButtonTooltip() const;
 	
 	TWeakPtr<FConsoleVariablesEditorListRow> Item;
-	
-	TSharedPtr<SButton> ActionButtonPtr;
-	TSharedPtr<SImage> ActionButtonImage;
 };
diff --git a/Engine/Plugins/Editor/ConsoleVariablesEditor/Source/ConsoleVariablesEditor/Private/Views/List/SConsoleVariablesEditorListValueInput.cpp b/Engine/Plugins/Editor/ConsoleVariablesEditor/Source/ConsoleVariablesEditor/Private/Views/List/SConsoleVariablesEditorListValueInput.cpp
--- a/Engine/Plugins/Editor/ConsoleVariablesEditor/Source/ConsoleVariablesEditor/Private/Views/List/SConsoleVariablesEditorListValueInput.cpp
+++ b/Engine/Plugins/Editor/ConsoleVariablesEditor/Source/ConsoleVariablesEditor/Private/Views/List/SConsoleVariablesEditorListValueInput.cpp
@@ -139,10 +139,17 @@ void SConsoleVariablesEditorListValueInput_Float::OnSliderAffected(const float I
 	const TSharedPtr<FConsoleVariablesEditorListRow> PinnedItem = Item.Pin();
 					
 	const FString ValueAsString = FString::SanitizeFloat(InValue);
-					
-	PinnedItem->GetCommandInfo().Pin()->ExecuteCommand(ValueAsString, true, true, !bPrintCommand);
-					
-	PinnedItem->SetCachedValue(ValueAsString);
+
+	for (const FConsoleVariablesEditorListRowPtr& RowPtr : PinnedItem->GetRowsAffectedByActions())
+	{
+		if (RowPtr->GetCommandInfo().Pin()->ObjectType != FConsoleVariablesEditorCommandInfo::EConsoleObjectType::Variable)
+		{
+			continue;
+		}
+
+		RowPtr->GetCommandInfo().Pin()->ExecuteCommand(ValueAsString, true, true, !bPrintCommand);
+		RowPtr->SetCachedValue(ValueAsString);
+	}
 }
 
 void SConsoleVariablesEditorListValueInput_Int::Construct(const FArguments& InArgs,
@@ -241,21 +248,28 @@ void SConsoleVariablesEditorListValueInput_Int::OnSliderAffected(const int32 InV
 	const TSharedPtr<FConsoleVariablesEditorListRow> PinnedItem = Item.Pin();
 					
 	const FString ValueAsString = FString::FromInt(InValue);
-					
-	if (!PinnedItem->GetCachedValue().Equals(ValueAsString))
-	{				
-		PinnedItem->GetCommandInfo().Pin()->ExecuteCommand(ValueAsString, true, true, !bPrintCommand);
-					
-		PinnedItem->SetCachedValue(ValueAsString);
 
-		return;
-	}
-
-	// If the new value and cached value aren't different the command won't be executed.
-	// If we still want to print the command, do it here
-	if (bPrintCommand)
+	for (const FConsoleVariablesEditorListRowPtr& RowPtr : PinnedItem->GetRowsAffectedByActions())
 	{
-		PinnedItem->GetCommandInfo().Pin()->PrintCommandOrVariable();
+		if (RowPtr->GetCommandInfo().Pin()->ObjectType != FConsoleVariablesEditorCommandInfo::EConsoleObjectType::Variable)
+		{
+			continue;
+		}
+
+		if (!RowPtr->GetCachedValue().Equals(ValueAsString))
+		{
+			RowPtr->GetCommandInfo().Pin()->ExecuteCommand(ValueAsString, true, true, !bPrintCommand);
+			RowPtr->SetCachedValue(ValueAsString);
+
+			continue;
+		}
+
+		// If the new value and cached value aren't different the command won't be executed.
+		// If we still want to print the command, do it here
+		if (bPrintCommand)
+		{
+			RowPtr->GetCommandInfo().Pin()->PrintCommandOrVariable();
+		}
 	}
 }
 
@@ -290,15 +304,23 @@ void SConsoleVariablesEditorListValueInput_String::Construct(const FArguments& I
 		.OnTextCommitted_Lambda([this] (const FText& InValue, ETextCommit::Type InTextCommitType)
 		{
 			check (Item.IsValid());
-			
+
+			const TSharedPtr<FConsoleVariablesEditorListRow> PinnedItem = Item.Pin();
+
 			const FString ValueAsString = InValue.ToString();
-			
-			if (const TSharedPtr<FConsoleVariablesEditorListRow> PinnedItem = Item.Pin();
-				!PinnedItem->GetCachedValue().Equals(ValueAsString))
-			{				
-				PinnedItem->GetCommandInfo().Pin()->ExecuteCommand(ValueAsString);
 
-				PinnedItem->SetCachedValue(ValueAsString);
+			for (const FConsoleVariablesEditorListRowPtr& RowPtr : PinnedItem->GetRowsAffectedByActions())
+			{
+				if (RowPtr->GetCommandInfo().Pin()->ObjectType != FConsoleVariablesEditorCommandInfo::EConsoleObjectType::Variable)
+				{
+					continue;
+				}
+
+				if (!RowPtr->GetCachedValue().Equals(ValueAsString))
+				{
+					RowPtr->GetCommandInfo().Pin()->ExecuteCommand(ValueAsString);
+					RowPtr->SetCachedValue(ValueAsString);
+				}
 			}
 		})
 		.IsEnabled(this, &SConsoleVariablesEditorListValueInput::IsRowChecked)
@@ -382,13 +404,21 @@ SConsoleVariablesEditorListValueInput_Bool::~SConsoleVariablesEditorListValueInp
 void SConsoleVariablesEditorListValueInput_Bool::SetInputValue(const FString& InValueAsString)
 {
 	check (Item.IsValid());
-			
-	if (const TSharedPtr<FConsoleVariablesEditorListRow> PinnedItem = Item.Pin();
-		!PinnedItem->GetCachedValue().Equals(InValueAsString))
-	{				
-		PinnedItem->GetCommandInfo().Pin()->ExecuteCommand(InValueAsString);
 
-		PinnedItem->SetCachedValue(InValueAsString);
+	const TSharedPtr<FConsoleVariablesEditorListRow> PinnedItem = Item.Pin();
+	
+	for (const FConsoleVariablesEditorListRowPtr& RowPtr : PinnedItem->GetRowsAffectedByActions())
+	{
+		if (RowPtr->GetCommandInfo().Pin()->ObjectType != FConsoleVariablesEditorCommandInfo::EConsoleObjectType::Variable)
+		{
+			continue;
+		}
+
+		if (!RowPtr->GetCachedValue().Equals(InValueAsString))
+		{
+			RowPtr->GetCommandInfo().Pin()->ExecuteCommand(InValueAsString);
+			RowPtr->SetCachedValue(InValueAsString);
+		}
 	}
 }
 
