[Console Variables Editor] Add toggle for loading presets as additive/replacement to existing list

diff --git a/Engine/Plugins/Editor/ConsoleVariablesEditor/Source/ConsoleVariablesEditor/Private/Views/MainPanel/ConsoleVariablesEditorMainPanel.cpp b/Engine/Plugins/Editor/ConsoleVariablesEditor/Source/ConsoleVariablesEditor/Private/Views/MainPanel/ConsoleVariablesEditorMainPanel.cpp
--- a/Engine/Plugins/Editor/ConsoleVariablesEditor/Source/ConsoleVariablesEditor/Private/Views/MainPanel/ConsoleVariablesEditorMainPanel.cpp
+++ b/Engine/Plugins/Editor/ConsoleVariablesEditor/Source/ConsoleVariablesEditor/Private/Views/MainPanel/ConsoleVariablesEditorMainPanel.cpp
@@ -7,6 +7,7 @@
 #include "ConsoleVariablesEditorCommandInfo.h"
 #include "ConsoleVariablesEditorLog.h"
 #include "ConsoleVariablesEditorModule.h"
+#include "ConsoleVariablesEditorProjectSettings.h"
 #include "HAL/IConsoleManager.h"
 #include "MultiUser/ConsoleVariableSync.h"
 #include "Views/List/ConsoleVariablesEditorList.h"
@@ -305,10 +306,46 @@ bool FConsoleVariablesEditorMainPanel::ImportPreset_Impl(
 {
 	if (Preset && EditingAsset)
 	{
-		ReferenceAssetOnDisk = Preset;
+		bool bReplaceExisting = false;
+
+		FConsoleVariablesEditorModule& CVarModule = GetConsoleVariablesModule();
+		if (UConsoleVariablesEditorProjectSettings* ProjectSettingsPtr =
+			GetMutableDefault<UConsoleVariablesEditorProjectSettings>())
+		{
+			bReplaceExisting = ProjectSettingsPtr->PresetImportMode == 
+				EConsoleVariablesEditorPresetImportMode::ReplaceExisting;
+		}
+
+		if (bReplaceExisting || EditingAsset->GetSavedCommands().IsEmpty())
+		{
+			ReferenceAssetOnDisk = Preset;
+		}
+		else if (ReferenceAssetOnDisk.Get() != EditingAsset.Get())
+		{
+			ReferenceAssetOnDisk.Reset();
+		}
 
 		EditingAsset->Modify();
-		EditingAsset->CopyFrom(Preset);
+
+		if (bReplaceExisting)
+		{
+			// Reset the editing asset's variable values before replacing the list
+			for (const FConsoleVariablesEditorAssetSaveData& CommandData : EditingAsset->GetSavedCommands())
+			{
+				const TWeakPtr<FConsoleVariablesEditorCommandInfo>& Command = CVarModule.FindCommandInfoByName(CommandData.CommandName);
+				if (const TSharedPtr<FConsoleVariablesEditorCommandInfo> PinnedCommand = Command.Pin())
+				{
+					PinnedCommand->SetSourceFlag(PinnedCommand->StartupSource);
+					PinnedCommand->ExecuteCommand(PinnedCommand->StartupValueAsString, true, false);
+				}
+			}
+
+			EditingAsset->CopyFrom(Preset);
+		}
+		else
+		{
+			EditingAsset->AddFrom(Preset);
+		}
 
 		return true;
 	}
diff --git a/Engine/Plugins/Editor/ConsoleVariablesEditor/Source/ConsoleVariablesEditor/Private/Views/MainPanel/SConsoleVariablesEditorMainPanel.cpp b/Engine/Plugins/Editor/ConsoleVariablesEditor/Source/ConsoleVariablesEditor/Private/Views/MainPanel/SConsoleVariablesEditorMainPanel.cpp
--- a/Engine/Plugins/Editor/ConsoleVariablesEditor/Source/ConsoleVariablesEditor/Private/Views/MainPanel/SConsoleVariablesEditorMainPanel.cpp
+++ b/Engine/Plugins/Editor/ConsoleVariablesEditor/Source/ConsoleVariablesEditor/Private/Views/MainPanel/SConsoleVariablesEditorMainPanel.cpp
@@ -5,6 +5,7 @@
 #include "ConsoleVariablesAsset.h"
 #include "ConsoleVariablesEditorLog.h"
 #include "ConsoleVariablesEditorModule.h"
+#include "ConsoleVariablesEditorProjectSettings.h"
 #include "ConsoleVariablesEditorStyle.h"
 #include "MultiUser/ConsoleVariableSyncData.h"
 #include "Views/List/ConsoleVariablesEditorList.h"
@@ -288,6 +289,74 @@ TSharedRef<SWidget> SConsoleVariablesEditorMainPanel::OnGeneratePresetsMenu()
 
 	MenuBuilder.BeginSection(NAME_None, LOCTEXT("ImportPreset_MenuSection", "Import Preset"));
 	{
+		MenuBuilder.AddMenuEntry(
+			LOCTEXT("ImportPreset_ImportMode_Add", "Add Preset to Existing Console Variables"),
+			LOCTEXT(
+				"ImportPreset_ImportMode_Add_Tooltip",
+				"Add the list of variables from the imported preset to the current preset, replacing the "
+				"values of any overlapping variables with the values from the imported preset."
+			),
+			FSlateIcon(),
+			FUIAction(
+				FExecuteAction::CreateLambda([this]() {
+					if (UConsoleVariablesEditorProjectSettings* ProjectSettingsPtr =
+						GetMutableDefault<UConsoleVariablesEditorProjectSettings>())
+					{
+						ProjectSettingsPtr->PresetImportMode =
+							EConsoleVariablesEditorPresetImportMode::AddToExisting;
+					}
+				}),
+				FCanExecuteAction(),
+				FIsActionChecked::CreateLambda([this]() {
+					if (const UConsoleVariablesEditorProjectSettings* ProjectSettingsPtr =
+						GetDefault<UConsoleVariablesEditorProjectSettings>())
+					{
+						return ProjectSettingsPtr->PresetImportMode ==
+							EConsoleVariablesEditorPresetImportMode::AddToExisting;
+					}
+
+					return false;
+				})
+			),
+			NAME_None,
+			EUserInterfaceActionType::RadioButton
+		);
+
+		MenuBuilder.AddMenuEntry(
+			LOCTEXT("ImportPreset_ImportMode_Replace", "Replace Console Variables with Preset"),
+			LOCTEXT(
+				"ImportPreset_ImportMode_Replace_Tooltip",
+				"Completely replace the list of variables in the current preset, resetting them "
+				"to their default values and removing them from the list before importing the new preset's variable list."
+			),
+			FSlateIcon(),
+			FUIAction(
+				FExecuteAction::CreateLambda([this]() {
+					if (UConsoleVariablesEditorProjectSettings* ProjectSettingsPtr =
+						GetMutableDefault<UConsoleVariablesEditorProjectSettings>())
+					{
+						ProjectSettingsPtr->PresetImportMode =
+							EConsoleVariablesEditorPresetImportMode::ReplaceExisting;
+					}
+				}),
+				FCanExecuteAction(),
+				FIsActionChecked::CreateLambda([this]() {
+					if (const UConsoleVariablesEditorProjectSettings* ProjectSettingsPtr =
+						GetDefault<UConsoleVariablesEditorProjectSettings>())
+					{
+						return ProjectSettingsPtr->PresetImportMode ==
+							EConsoleVariablesEditorPresetImportMode::ReplaceExisting;
+					}
+
+					return false;
+				})
+			),
+			NAME_None,
+			EUserInterfaceActionType::RadioButton
+		);
+
+		MenuBuilder.AddMenuSeparator();
+
 		TSharedRef<SWidget> PresetPicker = SNew(SBox)
 			.MinDesiredWidth(400.f)
 			.MinDesiredHeight(400.f)
diff --git a/Engine/Plugins/Editor/ConsoleVariablesEditor/Source/ConsoleVariablesEditor/Public/ConsoleVariablesEditorProjectSettings.h b/Engine/Plugins/Editor/ConsoleVariablesEditor/Source/ConsoleVariablesEditor/Public/ConsoleVariablesEditorProjectSettings.h
--- a/Engine/Plugins/Editor/ConsoleVariablesEditor/Source/ConsoleVariablesEditor/Public/ConsoleVariablesEditorProjectSettings.h
+++ b/Engine/Plugins/Editor/ConsoleVariablesEditor/Source/ConsoleVariablesEditor/Public/ConsoleVariablesEditorProjectSettings.h
@@ -11,6 +11,22 @@ enum class EConsoleVariablesEditorRowDisplayType : uint8
 	ShowLastEnteredValue
 };
 
+UENUM(BlueprintType)
+enum class EConsoleVariablesEditorPresetImportMode : uint8
+{
+	/**
+	 * Add the list of variables from the imported preset to the current preset, replacing the values of any overlapping
+	 * variables with the values from the imported preset.
+	 */
+	AddToExisting,
+
+	/**
+	 * Completely replace the list of variables in the current preset, resetting them to their default values and removing
+	 * them from the list before importing the new preset's variable list.
+	 */
+	ReplaceExisting
+};
+
 UCLASS(config = Engine, defaultconfig)
 class CONSOLEVARIABLESEDITOR_API UConsoleVariablesEditorProjectSettings : public UObject
 {
@@ -20,6 +36,7 @@ public:
 	UConsoleVariablesEditorProjectSettings(const FObjectInitializer& ObjectInitializer)
 	{
 		UncheckedRowDisplayType = EConsoleVariablesEditorRowDisplayType::ShowCurrentValue;
+		PresetImportMode = EConsoleVariablesEditorPresetImportMode::AddToExisting;
 		
 		bAddAllChangedConsoleVariablesToCurrentPreset = true;
 
@@ -38,6 +55,12 @@ public:
 	UPROPERTY(Config, EditAnywhere, Category="Console Variables Editor")
 	EConsoleVariablesEditorRowDisplayType UncheckedRowDisplayType;
 
+	/**
+	 * When importing a console variable preset, the variables 
+	 */
+	UPROPERTY(Config, EditAnywhere, Category = "Console Variables Editor")
+	EConsoleVariablesEditorPresetImportMode PresetImportMode;
+
 	/**
 	 *When variables are changed outside the Console Variables Editor, this option will add the variables to the current preset.
 	 *Does not apply to console commands like 'r.SetNearClipPlane' or 'stat fps'
diff --git a/Engine/Plugins/Editor/ConsoleVariablesEditor/Source/ConsoleVariablesEditorRuntime/Private/ConsoleVariablesAsset.cpp b/Engine/Plugins/Editor/ConsoleVariablesEditor/Source/ConsoleVariablesEditorRuntime/Private/ConsoleVariablesAsset.cpp
--- a/Engine/Plugins/Editor/ConsoleVariablesEditor/Source/ConsoleVariablesEditorRuntime/Private/ConsoleVariablesAsset.cpp
+++ b/Engine/Plugins/Editor/ConsoleVariablesEditor/Source/ConsoleVariablesEditorRuntime/Private/ConsoleVariablesAsset.cpp
@@ -142,3 +142,26 @@ void UConsoleVariablesAsset::CopyFrom(const UConsoleVariablesAsset* InAssetToCop
 	VariableCollectionDescription = InAssetToCopy->GetVariableCollectionDescription();
 	SavedCommands = InAssetToCopy->GetSavedCommands();
 }
+
+void UConsoleVariablesAsset::AddFrom(const UConsoleVariablesAsset* InAssetToCopy)
+{
+	for (const FConsoleVariablesEditorAssetSaveData& NewCommand : InAssetToCopy->GetSavedCommands())
+	{
+		FConsoleVariablesEditorAssetSaveData* ExistingCommand = SavedCommands.FindByPredicate(
+			[NewCommand](const FConsoleVariablesEditorAssetSaveData& ExistingCommand)
+			{
+				return ExistingCommand.CommandName.Compare(NewCommand.CommandName) == 0;
+			}
+		);
+
+		if (ExistingCommand)
+		{
+			ExistingCommand->CommandValueAsString = NewCommand.CommandValueAsString;
+			ExistingCommand->CheckedState = ECheckBoxState::Checked;
+		}
+		else
+		{
+			SavedCommands.Add(NewCommand);
+		}
+	}
+}
diff --git a/Engine/Plugins/Editor/ConsoleVariablesEditor/Source/ConsoleVariablesEditorRuntime/Public/ConsoleVariablesAsset.h b/Engine/Plugins/Editor/ConsoleVariablesEditor/Source/ConsoleVariablesEditorRuntime/Public/ConsoleVariablesAsset.h
--- a/Engine/Plugins/Editor/ConsoleVariablesEditor/Source/ConsoleVariablesEditorRuntime/Public/ConsoleVariablesAsset.h
+++ b/Engine/Plugins/Editor/ConsoleVariablesEditor/Source/ConsoleVariablesEditorRuntime/Public/ConsoleVariablesAsset.h
@@ -96,9 +96,13 @@ public:
 	UFUNCTION(BlueprintCallable, Category="Console Variables Asset")
 	bool RemoveConsoleVariable(const FString& InCommandString);
 
-	/** Copy data from input asset to this asset */
+	/** Copy list of variables from input asset to this asset */
 	UFUNCTION(BlueprintCallable, Category="Console Variables Asset")
 	void CopyFrom(const UConsoleVariablesAsset* InAssetToCopy);
+
+	/** Add list of variables from input asset to this asset's list */
+	UFUNCTION(BlueprintCallable, Category = "Console Variables Asset")
+	void AddFrom(const UConsoleVariablesAsset* InAssetToCopy);
 	
 private:
 
