Update BP library function for loading a preset to support add / replace assets.

diff --git a/Engine/Plugins/Editor/ConsoleVariablesEditor/Source/ConsoleVariablesEditor/Private/ConsoleVariablesEditorFunctionLibrary.cpp b/Engine/Plugins/Editor/ConsoleVariablesEditor/Source/ConsoleVariablesEditor/Private/ConsoleVariablesEditorFunctionLibrary.cpp
--- a/Engine/Plugins/Editor/ConsoleVariablesEditor/Source/ConsoleVariablesEditor/Private/ConsoleVariablesEditorFunctionLibrary.cpp
+++ b/Engine/Plugins/Editor/ConsoleVariablesEditor/Source/ConsoleVariablesEditor/Private/ConsoleVariablesEditorFunctionLibrary.cpp
@@ -19,7 +19,7 @@ UConsoleVariablesAsset* UConsoleVariablesEditorFunctionLibrary::GetCurrentlyLoad
 	return ConsoleVariablesEditorModule.GetPresetAsset();
 }
 
-void UConsoleVariablesEditorFunctionLibrary::LoadPresetIntoConsoleVariablesEditor(const UConsoleVariablesAsset* InAsset)
+void UConsoleVariablesEditorFunctionLibrary::LoadPresetIntoConsoleVariablesEditor(const UConsoleVariablesAsset* InAsset, const EConsoleVariablesEditorPresetImportMode InImportMode)
 {
 	if (!InAsset)
 	{
@@ -31,7 +31,7 @@ void UConsoleVariablesEditorFunctionLibrary::LoadPresetIntoConsoleVariablesEdito
 	
 	const FConsoleVariablesEditorModule& ConsoleVariablesEditorModule = FConsoleVariablesEditorModule::Get();
 
-	ConsoleVariablesEditorModule.OpenConsoleVariablesDialogWithAssetSelected(FAssetData(InAsset));
+	ConsoleVariablesEditorModule.OpenConsoleVariablesDialogWithAssetSelected(FAssetData(InAsset), InImportMode);
 }
 
 bool UConsoleVariablesEditorFunctionLibrary::CopyCurrentListToAsset(UConsoleVariablesAsset* InAsset)
diff --git a/Engine/Plugins/Editor/ConsoleVariablesEditor/Source/ConsoleVariablesEditor/Private/ConsoleVariablesEditorModule.cpp b/Engine/Plugins/Editor/ConsoleVariablesEditor/Source/ConsoleVariablesEditor/Private/ConsoleVariablesEditorModule.cpp
--- a/Engine/Plugins/Editor/ConsoleVariablesEditor/Source/ConsoleVariablesEditor/Private/ConsoleVariablesEditorModule.cpp
+++ b/Engine/Plugins/Editor/ConsoleVariablesEditor/Source/ConsoleVariablesEditor/Private/ConsoleVariablesEditorModule.cpp
@@ -133,7 +133,7 @@ void FConsoleVariablesEditorModule::OpenConsoleVariablesDialogWithPreset(
 	}
 }
 
-void FConsoleVariablesEditorModule::OpenConsoleVariablesDialogWithAssetSelected(const FAssetData& InAssetData) const
+void FConsoleVariablesEditorModule::OpenConsoleVariablesDialogWithAssetSelected(const FAssetData& InAssetData, const EConsoleVariablesEditorPresetImportMode InImportMode) const
 {
 	if (InAssetData.IsValid())
 	{
@@ -142,7 +142,7 @@ void FConsoleVariablesEditorModule::OpenConsoleVariablesDialogWithAssetSelected(
 
 	if (MainPanel.IsValid())
 	{
-		MainPanel->ImportPreset(InAssetData);
+		MainPanel->ImportPreset(InAssetData, InImportMode);
 	}
 }
 
diff --git a/Engine/Plugins/Editor/ConsoleVariablesEditor/Source/ConsoleVariablesEditor/Private/Views/MainPanel/ConsoleVariablesEditorMainPanel.cpp b/Engine/Plugins/Editor/ConsoleVariablesEditor/Source/ConsoleVariablesEditor/Private/Views/MainPanel/ConsoleVariablesEditorMainPanel.cpp
--- a/Engine/Plugins/Editor/ConsoleVariablesEditor/Source/ConsoleVariablesEditor/Private/Views/MainPanel/ConsoleVariablesEditorMainPanel.cpp
+++ b/Engine/Plugins/Editor/ConsoleVariablesEditor/Source/ConsoleVariablesEditor/Private/Views/MainPanel/ConsoleVariablesEditorMainPanel.cpp
@@ -278,13 +278,13 @@ void FConsoleVariablesEditorMainPanel::SaveSpecificPresetAs(const TObjectPtr<UCo
 	}
 }
 
-void FConsoleVariablesEditorMainPanel::ImportPreset(const FAssetData& InPresetAsset)
+void FConsoleVariablesEditorMainPanel::ImportPreset(const FAssetData& InPresetAsset, const EConsoleVariablesEditorPresetImportMode InImportMode)
 {
 	FSlateApplication::Get().DismissAllMenus();
 	
 	if (UConsoleVariablesAsset* Preset = CastChecked<UConsoleVariablesAsset>(InPresetAsset.GetAsset()))
 	{
-		if (const TObjectPtr<UConsoleVariablesAsset> EditingAsset = GetEditingAsset(); ImportPreset_Impl(Preset, EditingAsset))
+		if (const TObjectPtr<UConsoleVariablesAsset> EditingAsset = GetEditingAsset(); ImportPreset_Impl(Preset, EditingAsset, InImportMode))
 		{
 			EditorList->RebuildList("", false);
 		}
@@ -295,25 +295,28 @@ void FConsoleVariablesEditorMainPanel::ImportPreset(const TObjectPtr<UConsoleVar
 {
 	FSlateApplication::Get().DismissAllMenus();
 
-	if (const TObjectPtr<UConsoleVariablesAsset> EditingAsset = GetEditingAsset(); ImportPreset_Impl(InPreset, EditingAsset))
+	if (const TObjectPtr<UConsoleVariablesAsset> EditingAsset = GetEditingAsset(); ImportPreset_Impl(InPreset, EditingAsset, EConsoleVariablesEditorPresetImportMode::UseDefault))
 	{
 		EditorList->RebuildList("", false);
 	}
 }
 
 bool FConsoleVariablesEditorMainPanel::ImportPreset_Impl(
-	const TObjectPtr<UConsoleVariablesAsset> Preset, const TObjectPtr<UConsoleVariablesAsset> EditingAsset)
+	const TObjectPtr<UConsoleVariablesAsset> Preset, const TObjectPtr<UConsoleVariablesAsset> EditingAsset, const EConsoleVariablesEditorPresetImportMode InImportMode)
 {
 	if (Preset && EditingAsset)
 	{
-		bool bReplaceExisting = false;
+		bool bReplaceExisting = InImportMode == EConsoleVariablesEditorPresetImportMode::ReplaceExisting ? true : false;
 
 		FConsoleVariablesEditorModule& CVarModule = GetConsoleVariablesModule();
-		if (UConsoleVariablesEditorProjectSettings* ProjectSettingsPtr =
-			GetMutableDefault<UConsoleVariablesEditorProjectSettings>())
+		if (InImportMode == EConsoleVariablesEditorPresetImportMode::UseDefault)
 		{
-			bReplaceExisting = ProjectSettingsPtr->PresetImportMode == 
-				EConsoleVariablesEditorPresetImportMode::ReplaceExisting;
+			if (UConsoleVariablesEditorProjectSettings* ProjectSettingsPtr =
+				GetMutableDefault<UConsoleVariablesEditorProjectSettings>())
+			{
+				bReplaceExisting = ProjectSettingsPtr->PresetImportMode ==
+					EConsoleVariablesEditorPresetImportMode::ReplaceExisting;
+			}
 		}
 
 		if (bReplaceExisting || EditingAsset->GetSavedCommands().IsEmpty())
diff --git a/Engine/Plugins/Editor/ConsoleVariablesEditor/Source/ConsoleVariablesEditor/Private/Views/MainPanel/ConsoleVariablesEditorMainPanel.h b/Engine/Plugins/Editor/ConsoleVariablesEditor/Source/ConsoleVariablesEditor/Private/Views/MainPanel/ConsoleVariablesEditorMainPanel.h
--- a/Engine/Plugins/Editor/ConsoleVariablesEditor/Source/ConsoleVariablesEditor/Private/Views/MainPanel/ConsoleVariablesEditorMainPanel.h
+++ b/Engine/Plugins/Editor/ConsoleVariablesEditor/Source/ConsoleVariablesEditor/Private/Views/MainPanel/ConsoleVariablesEditorMainPanel.h
@@ -5,8 +5,11 @@
 #include "ConsoleVariablesEditorModule.h"
 #include "MultiUser/ConsoleVariableSync.h"
 #include "Views/List/ConsoleVariablesEditorList.h"
+#include "ConsoleVariablesEditorProjectSettings.h"
+
 #include "Widgets/SWidget.h"
 
+
 class FConsoleVariablesEditorList;
 class SConsoleVariablesEditorMainPanel;
 class UConsoleVariablesAsset;
@@ -55,7 +58,7 @@ public:
 	void SaveSpecificPreset(const TObjectPtr<UConsoleVariablesAsset> InPreset) const;
 	void SaveCurrentPresetAs();
 	void SaveSpecificPresetAs(const TObjectPtr<UConsoleVariablesAsset> InPreset) const;
-	void ImportPreset(const FAssetData& InPresetAsset);
+	void ImportPreset(const FAssetData& InPresetAsset, const EConsoleVariablesEditorPresetImportMode InImportMode = EConsoleVariablesEditorPresetImportMode::UseDefault);
 	void ImportPreset(const TObjectPtr<UConsoleVariablesAsset> InPreset);
 
 	TWeakObjectPtr<UConsoleVariablesAsset> GetReferenceAssetOnDisk() const
@@ -80,7 +83,7 @@ public:
 
 private:
 
-	bool ImportPreset_Impl(const TObjectPtr<UConsoleVariablesAsset> Preset, const TObjectPtr<UConsoleVariablesAsset> EditingAsset);
+	bool ImportPreset_Impl(const TObjectPtr<UConsoleVariablesAsset> Preset, const TObjectPtr<UConsoleVariablesAsset> EditingAsset, const EConsoleVariablesEditorPresetImportMode InImportMode);
 
 	TSharedPtr<SConsoleVariablesEditorMainPanel> MainPanelWidget;
 
diff --git a/Engine/Plugins/Editor/ConsoleVariablesEditor/Source/ConsoleVariablesEditor/Private/Views/MainPanel/SConsoleVariablesEditorMainPanel.cpp b/Engine/Plugins/Editor/ConsoleVariablesEditor/Source/ConsoleVariablesEditor/Private/Views/MainPanel/SConsoleVariablesEditorMainPanel.cpp
--- a/Engine/Plugins/Editor/ConsoleVariablesEditor/Source/ConsoleVariablesEditor/Private/Views/MainPanel/SConsoleVariablesEditorMainPanel.cpp
+++ b/Engine/Plugins/Editor/ConsoleVariablesEditor/Source/ConsoleVariablesEditor/Private/Views/MainPanel/SConsoleVariablesEditorMainPanel.cpp
@@ -283,8 +283,14 @@ TSharedRef<SWidget> SConsoleVariablesEditorMainPanel::OnGeneratePresetsMenu()
 		AssetPickerConfig.AssetShowWarningText = LOCTEXT("NoPresets_Warning", "No Presets Found");
 		AssetPickerConfig.Filter.ClassPaths.Add(UConsoleVariablesAsset::StaticClass()->GetClassPathName());
 		AssetPickerConfig.Filter.bRecursiveClasses = true;
-		AssetPickerConfig.OnAssetSelected =
-			FOnAssetSelected::CreateRaw(MainPanel.Pin().Get(), &FConsoleVariablesEditorMainPanel::ImportPreset);
+		AssetPickerConfig.OnAssetSelected = FOnAssetSelected::CreateLambda(
+			[this](const FAssetData& InAssetData)
+			{
+				if (MainPanel.IsValid())
+				{
+					MainPanel.Pin()->ImportPreset(InAssetData);
+				}
+			});
 	}
 
 	MenuBuilder.BeginSection(NAME_None, LOCTEXT("ImportPreset_MenuSection", "Import Preset"));
diff --git a/Engine/Plugins/Editor/ConsoleVariablesEditor/Source/ConsoleVariablesEditor/Public/ConsoleVariablesEditorFunctionLibrary.h b/Engine/Plugins/Editor/ConsoleVariablesEditor/Source/ConsoleVariablesEditor/Public/ConsoleVariablesEditorFunctionLibrary.h
--- a/Engine/Plugins/Editor/ConsoleVariablesEditor/Source/ConsoleVariablesEditor/Public/ConsoleVariablesEditorFunctionLibrary.h
+++ b/Engine/Plugins/Editor/ConsoleVariablesEditor/Source/ConsoleVariablesEditor/Public/ConsoleVariablesEditorFunctionLibrary.h
@@ -5,6 +5,8 @@
 #include "Containers/Array.h"
 #include "Kismet/BlueprintFunctionLibrary.h"
 
+#include "ConsoleVariablesEditorProjectSettings.h"
+
 #include "ConsoleVariablesEditorFunctionLibrary.generated.h"
 
 class FString;
@@ -23,7 +25,7 @@ public:
 
 	/** Loads the given asset in the Console Variables Editor and sets all its variable values. */
 	UFUNCTION(BlueprintCallable, Category = "Console Variables Editor")
-	static void LoadPresetIntoConsoleVariablesEditor(const UConsoleVariablesAsset* InAsset);
+	static void LoadPresetIntoConsoleVariablesEditor(const UConsoleVariablesAsset* InAsset, const EConsoleVariablesEditorPresetImportMode InImportMode = EConsoleVariablesEditorPresetImportMode::AddToExisting);
 
 	/*
 	 * Saves the current list in the Console Variables Editor to the given asset.
diff --git a/Engine/Plugins/Editor/ConsoleVariablesEditor/Source/ConsoleVariablesEditor/Public/ConsoleVariablesEditorModule.h b/Engine/Plugins/Editor/ConsoleVariablesEditor/Source/ConsoleVariablesEditor/Public/ConsoleVariablesEditorModule.h
--- a/Engine/Plugins/Editor/ConsoleVariablesEditor/Source/ConsoleVariablesEditor/Public/ConsoleVariablesEditorModule.h
+++ b/Engine/Plugins/Editor/ConsoleVariablesEditor/Source/ConsoleVariablesEditor/Public/ConsoleVariablesEditorModule.h
@@ -6,6 +6,7 @@
 #include "HAL/IConsoleManager.h"
 #include "Modules/ModuleInterface.h"
 #include "Widgets/Docking/SDockTab.h"
+#include "ConsoleVariablesEditorProjectSettings.h"
 
 struct FAssetData;
 struct FConsoleVariablesEditorCommandInfo;
@@ -34,7 +35,7 @@ public:
 	void SavePresetAs() const;
 	void SaveSpecificPresetAs(const TObjectPtr<UConsoleVariablesAsset> Preset) const;
 	void OpenConsoleVariablesDialogWithPreset(const TObjectPtr<UConsoleVariablesAsset> Preset) const;
-	void OpenConsoleVariablesDialogWithAssetSelected(const FAssetData& InAssetData) const;
+	void OpenConsoleVariablesDialogWithAssetSelected(const FAssetData& InAssetData, const EConsoleVariablesEditorPresetImportMode InImportMode = EConsoleVariablesEditorPresetImportMode::UseDefault) const;
 
 	/** Find all console variables and cache their startup values */
 	void QueryAndBeginTrackingConsoleVariables();
diff --git a/Engine/Plugins/Editor/ConsoleVariablesEditor/Source/ConsoleVariablesEditor/Public/ConsoleVariablesEditorProjectSettings.h b/Engine/Plugins/Editor/ConsoleVariablesEditor/Source/ConsoleVariablesEditor/Public/ConsoleVariablesEditorProjectSettings.h
--- a/Engine/Plugins/Editor/ConsoleVariablesEditor/Source/ConsoleVariablesEditor/Public/ConsoleVariablesEditorProjectSettings.h
+++ b/Engine/Plugins/Editor/ConsoleVariablesEditor/Source/ConsoleVariablesEditor/Public/ConsoleVariablesEditorProjectSettings.h
@@ -24,7 +24,12 @@ enum class EConsoleVariablesEditorPresetImportMode : uint8
 	 * Completely replace the list of variables in the current preset, resetting them to their default values and removing
 	 * them from the list before importing the new preset's variable list.
 	 */
-	ReplaceExisting
+	ReplaceExisting,
+
+	/**
+	 * Use the global import setting.
+	 */
+	UseDefault UMETA(Hidden)
 };
 
 UCLASS(config = Engine, defaultconfig)
