Remove check for slider changes as it is possible for the row to be destructed and a cancel commit value to be called which will trigger the lambda

diff --git a/Engine/Plugins/Editor/ConsoleVariablesEditor/Source/ConsoleVariablesEditor/Private/Views/List/SConsoleVariablesEditorListValueInput.cpp b/Engine/Plugins/Editor/ConsoleVariablesEditor/Source/ConsoleVariablesEditor/Private/Views/List/SConsoleVariablesEditorListValueInput.cpp
--- a/Engine/Plugins/Editor/ConsoleVariablesEditor/Source/ConsoleVariablesEditor/Private/Views/List/SConsoleVariablesEditorListValueInput.cpp
+++ b/Engine/Plugins/Editor/ConsoleVariablesEditor/Source/ConsoleVariablesEditor/Private/Views/List/SConsoleVariablesEditorListValueInput.cpp
@@ -135,7 +135,11 @@ float SConsoleVariablesEditorListValueInput_Float::GetInputValue() const
 
 void SConsoleVariablesEditorListValueInput_Float::OnSliderAffected(const float InValue, const bool bPrintCommand)
 {
-	check (Item.IsValid());
+	if (!Item.IsValid())
+	{
+		return;
+	}
+	
 	const TSharedPtr<FConsoleVariablesEditorListRow> PinnedItem = Item.Pin();
 					
 	const FString ValueAsString = FString::SanitizeFloat(InValue);
@@ -244,7 +248,11 @@ int32 SConsoleVariablesEditorListValueInput_Int::GetInputValue() const
 
 void SConsoleVariablesEditorListValueInput_Int::OnSliderAffected(const int32 InValue, const bool bPrintCommand)
 {
-	check (Item.IsValid());
+	if (!Item.IsValid())
+	{
+		return;
+	}
+
 	const TSharedPtr<FConsoleVariablesEditorListRow> PinnedItem = Item.Pin();
 					
 	const FString ValueAsString = FString::FromInt(InValue);
@@ -403,7 +411,10 @@ SConsoleVariablesEditorListValueInput_Bool::~SConsoleVariablesEditorListValueInp
 
 void SConsoleVariablesEditorListValueInput_Bool::SetInputValue(const FString& InValueAsString)
 {
-	check (Item.IsValid());
+	if (!Item.IsValid())
+	{
+		return;
+	}
 
 	const TSharedPtr<FConsoleVariablesEditorListRow> PinnedItem = Item.Pin();
 	
