Fix display of console variables when in Preview mode.

diff --git a/Engine/Plugins/Editor/ConsoleVariablesEditor/Source/ConsoleVariablesEditor/Private/Views/List/SConsoleVariablesEditorList.cpp b/Engine/Plugins/Editor/ConsoleVariablesEditor/Source/ConsoleVariablesEditor/Private/Views/List/SConsoleVariablesEditorList.cpp
--- a/Engine/Plugins/Editor/ConsoleVariablesEditor/Source/ConsoleVariablesEditor/Private/Views/List/SConsoleVariablesEditorList.cpp
+++ b/Engine/Plugins/Editor/ConsoleVariablesEditor/Source/ConsoleVariablesEditor/Private/Views/List/SConsoleVariablesEditorList.cpp
@@ -903,7 +903,8 @@ void SConsoleVariablesEditorList::SetupFilters()
 		ECVF_SetByConsoleVariablesIni,
 		ECVF_SetByCommandline,
 		ECVF_SetByCode,
-		ECVF_SetByConsole
+		ECVF_SetByConsole,
+		ECVF_SetByPreview
 	};
 
 	for (const EConsoleVariableFlags Type : SourceFilterTypes)
diff --git a/Engine/Plugins/Editor/ConsoleVariablesEditor/Source/ConsoleVariablesEditor/Public/ConsoleVariablesEditorCommandInfo.h b/Engine/Plugins/Editor/ConsoleVariablesEditor/Source/ConsoleVariablesEditor/Public/ConsoleVariablesEditorCommandInfo.h
--- a/Engine/Plugins/Editor/ConsoleVariablesEditor/Source/ConsoleVariablesEditor/Public/ConsoleVariablesEditorCommandInfo.h
+++ b/Engine/Plugins/Editor/ConsoleVariablesEditor/Source/ConsoleVariablesEditor/Public/ConsoleVariablesEditorCommandInfo.h
@@ -150,7 +150,8 @@ struct FConsoleVariablesEditorCommandInfo
 		{ EConsoleVariableFlags::ECVF_SetByConsoleVariablesIni, LOCTEXT("Source_SetByConsoleVariablesIni", "Console Variables ini") },
 		{ EConsoleVariableFlags::ECVF_SetByCommandline, LOCTEXT("Source_SetByCommandline", "Command line") },
 		{ EConsoleVariableFlags::ECVF_SetByCode, LOCTEXT("Source_SetByCode", "Code") },
-		{ EConsoleVariableFlags::ECVF_SetByConsole, LOCTEXT("Source_SetByConsole", "Console") }
+		{ EConsoleVariableFlags::ECVF_SetByConsole, LOCTEXT("Source_SetByConsole", "Console") },
+		{ EConsoleVariableFlags::ECVF_SetByPreview, LOCTEXT("Source_SetByPreview", "Preview") }
 	};
 };
 
