Check to see if the Multi-user plugin is available before loading it.

#Synced-CL 52105088

diff --git a/Engine/Plugins/Editor/ConsoleVariablesEditor/Source/ConsoleVariablesEditor/Private/MultiUser/ConcertConsoleVariableSessionCustomization.cpp b/Engine/Plugins/Editor/ConsoleVariablesEditor/Source/ConsoleVariablesEditor/Private/MultiUser/ConcertConsoleVariableSessionCustomization.cpp
--- a/Engine/Plugins/Editor/ConsoleVariablesEditor/Source/ConsoleVariablesEditor/Private/MultiUser/ConcertConsoleVariableSessionCustomization.cpp
+++ b/Engine/Plugins/Editor/ConsoleVariablesEditor/Source/ConsoleVariablesEditor/Private/MultiUser/ConcertConsoleVariableSessionCustomization.cpp
@@ -291,6 +291,12 @@ FConcertCVarDetails  GetLocalSessionConcertClientAsSetting(
 void FConcertConsoleVariableSessionCustomization::PopulateClientList()
 {
 	Clients.Reset();
+
+	if (!IConcertSyncClientModule::IsAvailable())
+	{
+		return;
+	}
+
 	TSharedPtr<IConcertSyncClient> ConcertSyncClient = IConcertSyncClientModule::Get().GetClient(TEXT("MultiUser"));
 	IConcertClientRef ConcertClient = ConcertSyncClient->GetConcertClient();
 	TSharedPtr<IConcertClientSession> Session = ConcertClient->GetCurrentSession();
@@ -395,6 +401,11 @@ void FConcertConsoleVariableSessionCustomization::CustomizeDetails(IDetailLayout
 		auto HandleGenerateRow = [this](TSharedPtr<FConcertCVarDetails> InClientInfo,
 									const TSharedRef<STableViewBase>& OwnerTable) -> TSharedRef<ITableRow>
 		{
+			if (!IConcertSyncClientModule::IsAvailable())
+			{
+				return SNew(STableRow<TSharedPtr<FConcertCVarDetails>>, OwnerTable);
+			}
+
 			TSharedPtr<IConcertSyncClient> ConcertSyncClient = IConcertSyncClientModule::Get().GetClient(TEXT("MultiUser"));
 			IConcertClientRef ConcertClient = ConcertSyncClient->GetConcertClient();
 			TSharedPtr<IConcertClientSession> ConcertClientSession = ConcertClient->GetCurrentSession();
diff --git a/Engine/Plugins/Editor/ConsoleVariablesEditor/Source/ConsoleVariablesEditor/Private/MultiUser/ConsoleVariableSync.cpp b/Engine/Plugins/Editor/ConsoleVariablesEditor/Source/ConsoleVariablesEditor/Private/MultiUser/ConsoleVariableSync.cpp
--- a/Engine/Plugins/Editor/ConsoleVariablesEditor/Source/ConsoleVariablesEditor/Private/MultiUser/ConsoleVariableSync.cpp
+++ b/Engine/Plugins/Editor/ConsoleVariablesEditor/Source/ConsoleVariablesEditor/Private/MultiUser/ConsoleVariableSync.cpp
@@ -29,22 +29,25 @@ struct FManagerImpl
 {
 	FManagerImpl()
 	{
-		if (TSharedPtr<IConcertSyncClient> ConcertSyncClient = IConcertSyncClientModule::Get().GetClient(TEXT("MultiUser")))
+		if (IConcertSyncClientModule::IsAvailable())
 		{
-			IConcertClientRef ConcertClient = ConcertSyncClient->GetConcertClient();
-			ConcertClient->OnSessionStartup().AddRaw(this, &FManagerImpl::Register);
-			ConcertClient->OnSessionShutdown().AddRaw(this, &FManagerImpl::Unregister);
-
-			if (TSharedPtr<IConcertClientSession> ConcertClientSession = ConcertClient->GetCurrentSession())
+			if (TSharedPtr<IConcertSyncClient> ConcertSyncClient = IConcertSyncClientModule::Get().GetClient(TEXT("MultiUser")))
 			{
-				Register(ConcertClientSession.ToSharedRef());
-			}
+				IConcertClientRef ConcertClient = ConcertSyncClient->GetConcertClient();
+				ConcertClient->OnSessionStartup().AddRaw(this, &FManagerImpl::Register);
+				ConcertClient->OnSessionShutdown().AddRaw(this, &FManagerImpl::Unregister);
 
-			ConcertClient->OnSessionConnectionChanged().AddRaw(this, &FManagerImpl::OnSessionConnectionChanged);
+				if (TSharedPtr<IConcertClientSession> ConcertClientSession = ConcertClient->GetCurrentSession())
+				{
+					Register(ConcertClientSession.ToSharedRef());
+				}
 
-			RegisterExtensions();
+				ConcertClient->OnSessionConnectionChanged().AddRaw(this, &FManagerImpl::OnSessionConnectionChanged);
 
-			bIsInitialized = true;
+				RegisterExtensions();
+
+				bIsInitialized = true;
+			}
 		}
 	}
 
