[ConsoleVariables] Fix issue with console variables values not being restored to the original value when unchecking all rows at once instead one by one

#Synced-CL 52173051

diff --git a/Engine/Plugins/Editor/ConsoleVariablesEditor/Source/ConsoleVariablesEditor/Private/Views/List/SConsoleVariablesEditorList.cpp b/Engine/Plugins/Editor/ConsoleVariablesEditor/Source/ConsoleVariablesEditor/Private/Views/List/SConsoleVariablesEditorList.cpp
--- a/Engine/Plugins/Editor/ConsoleVariablesEditor/Source/ConsoleVariablesEditor/Private/Views/List/SConsoleVariablesEditorList.cpp
+++ b/Engine/Plugins/Editor/ConsoleVariablesEditor/Source/ConsoleVariablesEditor/Private/Views/List/SConsoleVariablesEditorList.cpp
@@ -835,6 +835,14 @@ TSharedPtr<SHeaderRow> SConsoleVariablesEditorList::GenerateHeaderRow()
 							FConsoleVariablesEditorCommandInfo::EConsoleObjectType::Variable)
 						{
 							Object->SetWidgetCheckedState(NewState);
+							if (NewState == ECheckBoxState::Unchecked)
+							{
+								Object->ResetToStartupValueAndSource();
+							}
+							else if (NewState == ECheckBoxState::Checked)
+							{
+								Object->ResetToPresetValue();
+							}
 						}
 					}
 				})
