Enable Console Variable Editor to be a default plugin.  Move it to production ready.

#Synced-CL 52173051

diff --git a/Engine/Plugins/Editor/ConsoleVariablesEditor/ConsoleVariables.uplugin b/Engine/Plugins/Editor/ConsoleVariablesEditor/ConsoleVariables.uplugin
--- a/Engine/Plugins/Editor/ConsoleVariablesEditor/ConsoleVariables.uplugin
+++ b/Engine/Plugins/Editor/ConsoleVariablesEditor/ConsoleVariables.uplugin
@@ -8,11 +8,11 @@
 	"CreatedBy": "Epic Games, Inc.",
 	"CreatedByURL": "https://epicgames.com",
 	"DocsURL": "",
-	"EnabledByDefault": false,
+	"EnabledByDefault": true,
 	"MarketplaceURL": "",
 	"SupportURL": "",
 	"CanContainContent": true,
-	"IsBetaVersion": true,
+	"IsBetaVersion": false,
 	"IsExperimentalVersion": false,
 	"Installed": false,
 	"Plugins":
